Implement the slot layer for the UI framework.

diff --git a/Engine/Plugins/Experimental/UIFramework/Source/Private/Layers/UIFCanvasLayer.cpp b/Engine/Plugins/Experimental/UIFramework/Source/Private/Layers/UIFCanvasLayer.cpp
--- a/Engine/Plugins/Experimental/UIFramework/Source/Private/Layers/UIFCanvasLayer.cpp
+++ b/Engine/Plugins/Experimental/UIFramework/Source/Private/Layers/UIFCanvasLayer.cpp
@@ -235,7 +235,7 @@ void UUIFCanvasLayer::OnLocalLayerWidgetAdded()
 	UUIFCanvasLayerUserWidget* LocalLayerWidget = CastChecked<UUIFCanvasLayerUserWidget>(GetLayerWidget());
 	for (const FUIFCanvasWidgetEntry& Entry : WidgetList.GetEntries())
 	{
-		if (Entry.Widget && Entry.Widget->GetWidget())
+		if (Entry.Widget && Entry.Widget->GetWidget() && Entry.bAdded)
 		{
 			LocalLayerWidget->AddWidget(Entry.Widget->GetWidget(), Entry.Slot);
 		}
diff --git a/Engine/Plugins/Experimental/UIFramework/Source/Private/Layers/UIFSlotLayer.cpp b/Engine/Plugins/Experimental/UIFramework/Source/Private/Layers/UIFSlotLayer.cpp
new file mode 100644
--- /dev/null
+++ b/Engine/Plugins/Experimental/UIFramework/Source/Private/Layers/UIFSlotLayer.cpp
@@ -0,0 +1,240 @@
+// Copyright Epic Games, Inc. All Rights Reserved.
+
+#include "Layers/UIFSlotLayer.h"
+
+#include "UIFPlayerComponent.h"
+#include "UIFWidget.h"
+
+#include "Blueprint/UserWidget.h"
+#include "Engine/ActorChannel.h"
+#include "Engine/AssetManager.h"
+#include "Engine/Engine.h"
+#include "Engine/NetDriver.h"
+#include "Engine/StreamableManager.h"
+#include "GameFramework/Actor.h"
+#include "GameFramework/PlayerController.h"
+#include "Net/UnrealNetwork.h"
+
+
+ /**
+ *
+ */
+void FUIFSlotWidgetList::PreReplicatedRemove(const TArrayView<int32> RemovedIndices, int32 FinalSize)
+{
+	for (int32 Index : RemovedIndices)
+	{
+		FUIFSlotWidgetEntry& Entry = Entries[Index];
+		if (Entry.Widget && !Entry.SlotName.IsNone() && Entry.bAdded)
+		{
+			OwnerLayer->LocalRemoveWidget(Entry.Widget, Entry.SlotName);
+			Entry.bAdded = false;
+		}
+	}
+}
+
+
+void FUIFSlotWidgetList::PostReplicatedAdd(const TArrayView<int32> AddedIndices, int32 FinalSize)
+{
+	for (int32 Index : AddedIndices)
+	{
+		FUIFSlotWidgetEntry& Entry = Entries[Index];
+		if (Entry.Widget && !Entry.SlotName.IsNone())
+		{
+			ensureMsgf(!Entry.bAdded, TEXT("The widget was already added."));
+			OwnerLayer->LocalAddWidget(Entry.Widget, Entry.SlotName);
+			Entry.bAdded = true;
+		}
+	}
+}
+
+
+void FUIFSlotWidgetList::PostReplicatedChange(const TArrayView<int32>& ChangedIndices, int32 FinalSize)
+{
+	for (int32 Index : ChangedIndices)
+	{
+		FUIFSlotWidgetEntry& Entry = Entries[Index];
+		if (Entry.Widget && !Entry.SlotName.IsNone() && !Entry.bAdded)
+		{
+			OwnerLayer->LocalAddWidget(Entry.Widget, Entry.SlotName);
+			Entry.bAdded = true;
+		}
+		else if (!Entry.Widget && !Entry.SlotName.IsNone() && Entry.bAdded)
+		{
+			OwnerLayer->LocalRemoveEmptySlots();
+			Entry.bAdded = false;
+		}
+		else if (Entry.Widget && !Entry.SlotName.IsNone() && Entry.bAdded)
+		{
+			ensureMsgf(false, TEXT("SlotName cannot be changed at runtime."));
+		}
+	}
+}
+
+
+void FUIFSlotWidgetList::AddEntry(UUIFWidget* InWidget, FName InSlot)
+{
+	check(InWidget);
+	FUIFSlotWidgetEntry& NewEntry = Entries.AddDefaulted_GetRef();
+	NewEntry.Widget = InWidget;
+	NewEntry.SlotName = InSlot;
+	MarkItemDirty(NewEntry);
+}
+
+
+void FUIFSlotWidgetList::RemoveEntry(UUIFWidget* InWidget)
+{
+	check(InWidget);
+	int32 Index = Entries.IndexOfByPredicate([InWidget](const FUIFSlotWidgetEntry& Entry) { return Entry.Widget == InWidget; });
+	if (Index != INDEX_NONE)
+	{
+		Entries.RemoveAt(Index);
+		MarkArrayDirty();
+	}
+}
+
+
+/**
+ *
+ */
+UUIFSlotLayer::UUIFSlotLayer()
+	: WidgetList(this)
+{
+}
+
+
+void UUIFSlotLayer::GetLifetimeReplicatedProps(TArray<FLifetimeProperty>& OutLifetimeProps) const
+{
+	Super::GetLifetimeReplicatedProps(OutLifetimeProps);
+
+	FDoRepLifetimeParams Params;
+	Params.bIsPushBased = true;
+	DOREPLIFETIME_WITH_PARAMS_FAST(ThisClass, WidgetList, Params);
+}
+
+
+bool UUIFSlotLayer::ReplicateSubobjects(UActorChannel* Channel, class FOutBunch* Bunch, FReplicationFlags* RepFlags)
+{
+	bool WroteSomething = Super::ReplicateSubobjects(Channel, Bunch, RepFlags);
+
+	for (FUIFSlotWidgetEntry& Entry : WidgetList.Entries)
+	{
+		UUIFWidget* Widget = Entry.Widget;
+		if (IsValid(Widget))
+		{
+			WroteSomething = Channel->ReplicateSubobject(Widget, *Bunch, *RepFlags) || WroteSomething;
+		}
+	}
+
+	return WroteSomething;
+}
+
+
+UUIFWidget* UUIFSlotLayer::CreateWidget(TSubclassOf<UUIFWidget> WidgetClass, FName Slot)
+{
+	check(GetOuterAPlayerController()->HasAuthority());
+
+	UUIFWidget* Result = nullptr;
+	if (WidgetClass.Get() == nullptr)
+	{
+		return Result;
+	}
+
+	Result = NewObject<UUIFWidget>(GetOuterAPlayerController(), WidgetClass.Get());
+	WidgetList.AddEntry(Result, Slot);
+
+	return Result;
+}
+
+
+void UUIFSlotLayer::RemoveWidget(UUIFWidget* Widget)
+{
+	check(GetOuterAPlayerController()->HasAuthority());
+
+	if (Widget)
+	{
+		WidgetList.RemoveEntry(Widget);
+	}
+}
+
+
+void UUIFSlotLayer::LocalAddWidget(UUIFWidget* InUIWidget, FName InSlotName)
+{
+	checkf(!InSlotName.IsNone(), TEXT("The SlotName is invalid."));
+
+	LocalSlotContents.FindOrAdd(InSlotName) = InUIWidget;
+
+	TWeakObjectPtr<UUIFSlotLayer> WeakSelf = this;
+	TWeakObjectPtr<UUIFWidget> WeakUIWidget = InUIWidget;
+	InUIWidget->LocalCreateWidgetAsync([WeakSelf, WeakUIWidget, InSlotName]()
+		{
+			UUIFSlotLayer* StrongSelf = WeakSelf.Get();
+			UUIFWidget* StrongUIWidget = WeakUIWidget.Get();
+			if (StrongSelf && StrongUIWidget && StrongUIWidget->GetWidget() && StrongSelf->GetLayerWidget())
+			{
+				if (UUIFWidget** FoundWidget = StrongSelf->LocalSlotContents.Find(InSlotName))
+				{
+					if (*FoundWidget == StrongUIWidget) // was it changed by another call
+					{
+						StrongSelf->GetLayerWidget()->SetContentForSlot(InSlotName, StrongUIWidget->GetWidget());
+					}
+				}
+			}
+		});
+}
+
+
+void UUIFSlotLayer::LocalRemoveWidget(UUIFWidget* InUIWidget, FName InSlotName)
+{
+	checkf(!InSlotName.IsNone(), TEXT("The SlotName is invalid."));
+
+	if (UUIFWidget** FoundWidget = LocalSlotContents.Find(InSlotName))
+	{
+		if (*FoundWidget == InUIWidget) // was it changed by another async call
+		{
+			LocalSlotContents.Remove(InSlotName);
+
+			UWidget* PreviousContent = InUIWidget->GetWidget();
+			UUserWidget* LocalLayerWidget = GetLayerWidget();
+			if (LocalLayerWidget && PreviousContent)
+			{
+				const UWidget* OldContent = LocalLayerWidget->GetContentForSlot(InSlotName);
+				if (OldContent == PreviousContent)
+				{
+					LocalLayerWidget->SetContentForSlot(InSlotName, nullptr);
+				}
+			}
+		}
+	}
+}
+
+
+void UUIFSlotLayer::LocalRemoveEmptySlots()
+{
+	if (UUserWidget* LocalLayerWidget = GetLayerWidget())
+	{
+		for (auto Itt = LocalSlotContents.CreateIterator(); Itt; ++Itt)
+		{
+			const UWidget* OldContent = LocalLayerWidget->GetContentForSlot(Itt.Key());
+			if (OldContent == nullptr)
+			{
+				LocalLayerWidget->SetContentForSlot(Itt.Key(), nullptr);
+				Itt.RemoveCurrent();
+			}
+		}
+	}
+}
+
+
+void UUIFSlotLayer::OnLocalLayerWidgetAdded()
+{
+	if (UUserWidget* LocalLayerWidget = GetLayerWidget())
+	{
+		for (const FUIFSlotWidgetEntry& Entry : WidgetList.GetEntries())
+		{
+			if (Entry.Widget && Entry.Widget->GetWidget() && !Entry.SlotName.IsNone() && Entry.bAdded)
+			{
+				LocalLayerWidget->SetContentForSlot(Entry.SlotName, Entry.Widget->GetWidget());
+			}
+		}
+	}
+}
\ No newline at end of file
diff --git a/Engine/Plugins/Experimental/UIFramework/Source/Public/Layers/UIFSlotLayer.h b/Engine/Plugins/Experimental/UIFramework/Source/Public/Layers/UIFSlotLayer.h
new file mode 100644
--- /dev/null
+++ b/Engine/Plugins/Experimental/UIFramework/Source/Public/Layers/UIFSlotLayer.h
@@ -0,0 +1,119 @@
+// Copyright Epic Games, Inc. All Rights Reserved.
+
+#pragma once
+
+#include "UIFLayer.h"
+
+#include "UIFSlotLayer.generated.h"
+
+
+/**
+ *
+ */
+USTRUCT(BlueprintType)
+struct FUIFSlotWidgetEntry : public FFastArraySerializerItem
+{
+	GENERATED_BODY()
+
+	FUIFSlotWidgetEntry() = default;
+
+	UPROPERTY()
+	TObjectPtr<UUIFWidget> Widget = nullptr;
+	
+	UPROPERTY()
+	FName SlotName;
+
+	UPROPERTY(NotReplicated)
+	bool bAdded = false;
+};
+
+
+/**
+ * 
+ */
+USTRUCT(BlueprintType)
+struct UIFRAMEWORK_API FUIFSlotWidgetList : public FFastArraySerializer
+{
+	GENERATED_BODY()
+
+	friend UUIFSlotLayer;
+
+	FUIFSlotWidgetList() = default;
+
+	FUIFSlotWidgetList(UUIFSlotLayer* InOwnerLayer)
+		: OwnerLayer(InOwnerLayer)
+	{
+	}
+
+public:
+	//~ Begin of FFastArraySerializer
+	void PreReplicatedRemove(const TArrayView<int32> RemovedIndices, int32 FinalSize);
+	void PostReplicatedAdd(const TArrayView<int32> AddedIndices, int32 FinalSize);
+	void PostReplicatedChange(const TArrayView<int32>& ChangedIndices, int32 FinalSize);
+	//~ End of FFastArraySerializer
+
+	bool NetDeltaSerialize(FNetDeltaSerializeInfo& DeltaParms)
+	{
+		return FFastArraySerializer::FastArrayDeltaSerialize<FUIFSlotWidgetEntry, FUIFSlotWidgetList>(Entries, DeltaParms, *this);
+	}
+
+	void AddEntry(UUIFWidget* Widget, FName Slot);
+	void RemoveEntry(UUIFWidget* Widget);
+
+	TArrayView<const FUIFSlotWidgetEntry> GetEntries() const
+	{
+		return Entries;
+	}
+
+private:
+	UPROPERTY()
+	TArray<FUIFSlotWidgetEntry> Entries;
+
+	UPROPERTY(NotReplicated, Transient)
+	TObjectPtr<UUIFSlotLayer> OwnerLayer = nullptr;
+};
+
+template<>
+struct TStructOpsTypeTraits<FUIFSlotWidgetList> : public TStructOpsTypeTraitsBase2<FUIFSlotWidgetList>
+{
+	enum { WithNetDeltaSerializer = true };
+};
+
+
+/**
+ *
+ */
+UCLASS(Abstract, DisplayName="Slot Layer UIFramework")
+class UIFRAMEWORK_API UUIFSlotLayer : public UUIFLayer
+{
+	GENERATED_BODY()
+
+	friend FUIFSlotWidgetList;
+
+public:
+	UUIFSlotLayer();
+
+	//~ Begin UUILayer
+	virtual bool ReplicateSubobjects(UActorChannel* Channel, class FOutBunch* Bunch, FReplicationFlags* RepFlags) override;
+protected:
+	virtual void OnLocalLayerWidgetAdded() override;
+	//~ End UUILayer
+
+public:
+	UFUNCTION(BlueprintCallable, BlueprintAuthorityOnly, Category = "UI Framework", meta = (DeterminesOutputType = WidgetClass))
+	UUIFWidget* CreateWidget(UPARAM(meta = (AllowAbstract = false)) TSubclassOf<UUIFWidget> WidgetClass, FName Slot);
+	UFUNCTION(BlueprintCallable, BlueprintAuthorityOnly, Category = "UI Framework")
+	void RemoveWidget(UUIFWidget* Widget);
+
+private:
+	void LocalAddWidget(UUIFWidget* Widget, FName SlotName);
+	void LocalRemoveWidget(UUIFWidget* Widget, FName SlotName);
+	void LocalRemoveEmptySlots();
+
+private:
+	UPROPERTY(Replicated)
+	FUIFSlotWidgetList WidgetList;
+
+	UPROPERTY(Transient)
+	TMap<FName, UUIFWidget*> LocalSlotContents;
+};
