UIFramework:
Add Canvas. Rework how the widget tree is managed on local.

diff --git a/Engine/Plugins/Experimental/UIFramework/Source/Private/Types/UIFSlotBase.cpp b/Engine/Plugins/Experimental/UIFramework/Source/Private/Types/UIFSlotBase.cpp
--- a/Engine/Plugins/Experimental/UIFramework/Source/Private/Types/UIFSlotBase.cpp
+++ b/Engine/Plugins/Experimental/UIFramework/Source/Private/Types/UIFSlotBase.cpp
@@ -16,7 +16,7 @@
 /**
  *
  */
-void FUIFrameworkSlotBase::SetWidget(UUIFrameworkWidget* InWidget)
+void FUIFrameworkSlotBase::AuthoritySetWidget(UUIFrameworkWidget* InWidget)
 {
 	Widget = InWidget;
 	WidgetId = Widget ? Widget->GetWidgetId() : FUIFrameworkWidgetId();
diff --git a/Engine/Plugins/Experimental/UIFramework/Source/Private/Types/UIFWidgetTree.cpp b/Engine/Plugins/Experimental/UIFramework/Source/Private/Types/UIFWidgetTree.cpp
--- a/Engine/Plugins/Experimental/UIFramework/Source/Private/Types/UIFWidgetTree.cpp
+++ b/Engine/Plugins/Experimental/UIFramework/Source/Private/Types/UIFWidgetTree.cpp
@@ -44,7 +44,28 @@ void FUIFrameworkWidgetTree::PreReplicatedRemove(const TArrayView<int32> Removed
 		FUIFrameworkWidgetTreeEntry& Entry = Entries[Index];
 		if (Entry.Child)
 		{
-			OwnerComponent->LocalRemoveWidget(Entry);
+			OwnerComponent->LocalWidgetRemovedFromTree(Entry);
+
+			// Was it remove but added by something else
+			{
+				int32 Count = 0;
+				for (const FUIFrameworkWidgetTreeEntry& Other : Entries)
+				{
+					if (Other.Child == Entry.Child)
+					{
+						++Count;
+						if (Count > 2)
+						{
+							break;
+						}
+					}
+				}
+				if (Count <= 1)
+				{
+					Entry.Child->LocalDestroyUMGWidget();
+					WidgetByIdMap.Remove(Entry.Child->GetWidgetId());
+				}
+			}
 		}
 	}
 }
@@ -56,7 +77,12 @@ void FUIFrameworkWidgetTree::PostReplicatedAdd(const TArrayView<int32> AddedIndi
 		FUIFrameworkWidgetTreeEntry& Entry = Entries[Index];
 		if (Entry.ParentId.IsValid() && Entry.ChildId.IsValid())
 		{
-			OwnerComponent->LocalAddWidget(Entry);
+			OwnerComponent->LocalWidgetWasAddedToTree(Entry);
+			WidgetByIdMap.FindOrAdd(Entry.ChildId) = Entry.Child;
+			if (!Entry.ParentId.IsRoot())
+			{
+				WidgetByIdMap.FindOrAdd(Entry.ParentId) = Entry.Parent;
+			}
 		}
 	}
 }
@@ -69,7 +95,12 @@ void FUIFrameworkWidgetTree::PostReplicatedChange(const TArrayView<int32>& Chang
 		FUIFrameworkWidgetTreeEntry& Entry = Entries[Index];
 		if (Entry.ParentId.IsValid() && Entry.ChildId.IsValid())
 		{
-			OwnerComponent->LocalAddWidget(Entry);
+			OwnerComponent->LocalWidgetWasAddedToTree(Entry);
+			WidgetByIdMap.FindOrAdd(Entry.ChildId) = Entry.Child;
+			if (!Entry.ParentId.IsRoot())
+			{
+				WidgetByIdMap.FindOrAdd(Entry.ParentId) = Entry.Parent;
+			}
 		}
 	}
 }
@@ -127,6 +158,7 @@ void FUIFrameworkWidgetTree::AddChildInternal(UUIFrameworkWidget* Parent, UUIFra
 	{
 		FUIFrameworkWidgetTreeEntry& NewEntry = Entries.Emplace_GetRef(Parent, Child);
 		MarkItemDirty(NewEntry);
+		WidgetByIdMap.FindOrAdd(Child->GetWidgetId()) = Child;
 		AddChildRecursiveInternal(Child);
 	}
 }
@@ -138,6 +170,7 @@ void FUIFrameworkWidgetTree::AddChildRecursiveInternal(UUIFrameworkWidget* Widge
 		{
 			FUIFrameworkWidgetTreeEntry& NewEntry = Self->Entries.Emplace_GetRef(Widget, ChildWidget);
 			Self->MarkItemDirty(NewEntry);
+			Self->WidgetByIdMap.FindOrAdd(ChildWidget->GetWidgetId()) = ChildWidget;
 			Self->AddChildRecursiveInternal(ChildWidget);
 		});
 }
@@ -154,6 +187,8 @@ void FUIFrameworkWidgetTree::RemoveWidget(UUIFrameworkWidget* Widget)
 		while (WidgetsToRemove.Num() > 0)
 		{
 			UUIFrameworkWidget* CurrentWidgetToRemove = WidgetsToRemove.Last();
+			WidgetByIdMap.Remove(CurrentWidgetToRemove->GetWidgetId());
+
 			WidgetsToRemove.RemoveAt(WidgetsToRemove.Num()-1);
 			for (int32 Index = Entries.Num() - 1; Index >= 0; --Index)
 			{
@@ -196,3 +231,13 @@ const FUIFrameworkWidgetTreeEntry* FUIFrameworkWidgetTree::GetEntryByReplication
 	return const_cast<FUIFrameworkWidgetTree*>(this)->GetEntryByReplicationId(ReplicationId);
 }
 
+UUIFrameworkWidget* FUIFrameworkWidgetTree::FindWidgetById(FUIFrameworkWidgetId WidgetId)
+{
+	const TWeakObjectPtr<UUIFrameworkWidget>* Found = WidgetByIdMap.Find(WidgetId);
+	return Found ? Found->Get() : nullptr;
+}
+
+const UUIFrameworkWidget* FUIFrameworkWidgetTree::FindWidgetById(FUIFrameworkWidgetId WidgetId) const
+{
+	return const_cast<FUIFrameworkWidgetTree*>(this)->FindWidgetById(WidgetId);
+}
\ No newline at end of file
diff --git a/Engine/Plugins/Experimental/UIFramework/Source/Private/UIFPlayerComponent.cpp b/Engine/Plugins/Experimental/UIFramework/Source/Private/UIFPlayerComponent.cpp
--- a/Engine/Plugins/Experimental/UIFramework/Source/Private/UIFPlayerComponent.cpp
+++ b/Engine/Plugins/Experimental/UIFramework/Source/Private/UIFPlayerComponent.cpp
@@ -3,6 +3,7 @@
 #include "UIFPlayerComponent.h"
 #include "UIFLog.h"
 #include "UIFWidget.h"
+#include "Blueprint/GameViewportSubsystem.h"
 
 #include "Blueprint/UserWidget.h"
 #include "Engine/ActorChannel.h"
@@ -24,11 +25,9 @@ void FUIFrameworkGameLayerSlotList::PostReplicatedChange(const TArrayView<int32>
 	for (int32 Index : ChangedIndices)
 	{
 		FUIFrameworkGameLayerSlot& Slot = Entries[Index];
-		if (Slot.GetWidget() && Slot.GetWidgetId() == Slot.LocalPreviousWidgetId && Slot.GetWidget()->LocalGetUMGWidget())
+		if (Slot.LocalIsAquiredWidgetValid())
 		{
-			// Remove and add the widget again
-			Slot.GetWidget()->LocalGetUMGWidget()->RemoveFromParent();
-			Owner->LocalAddChild(Slot.LocalPreviousWidgetId);
+			Owner->LocalAddChild(Slot.GetWidgetId());
 		}
 	}
 }
@@ -42,7 +41,7 @@ void FUIFrameworkGameLayerSlotList::AddEntry(FUIFrameworkGameLayerSlot Entry)
 bool FUIFrameworkGameLayerSlotList::RemoveEntry(UUIFrameworkWidget* Widget)
 {
 	check(Widget);
-	const int32 Index = Entries.IndexOfByPredicate([Widget](const FUIFrameworkGameLayerSlot& Entry){ return Entry.GetWidget() == Widget; });
+	const int32 Index = Entries.IndexOfByPredicate([Widget](const FUIFrameworkGameLayerSlot& Entry){ return Entry.AuthorityGetWidget() == Widget; });
 	if (Index != INDEX_NONE)
 	{
 		Entries.RemoveAt(Index);
@@ -79,9 +78,9 @@ void UUIFrameworkPlayerComponent::UninitializeComponent()
 	{
 		for (FUIFrameworkGameLayerSlot& Entry : RootList.Entries)
 		{
-			if (Entry.GetWidget())
+			if (UUIFrameworkWidget* ChildWidget = GetWidgetTree().FindWidgetById(Entry.GetWidgetId()))
 			{
-				Entry.GetWidget()->LocalDestroyUMGWidget();
+				ChildWidget->LocalDestroyUMGWidget();
 			}
 		}
 	}
@@ -109,21 +108,21 @@ void UUIFrameworkPlayerComponent::AddWidget(FUIFrameworkGameLayerSlot InEntry)
 	APlayerController* LocalOwner = GetPlayerController();
 	check(LocalOwner->HasAuthority());
 
-	if (InEntry.GetWidget() == nullptr)
+	if (InEntry.AuthorityGetWidget() == nullptr)
 	{
 		FFrame::KismetExecutionMessage(TEXT("The widget is invalid. It can't be added."), ELogVerbosity::Warning, "InvalidWidgetToAdd");
 	}
 	else
 	{
-		UUIFrameworkPlayerComponent* PreviousOwner = InEntry.GetWidget()->GetPlayerComponent();
+		UUIFrameworkPlayerComponent* PreviousOwner = InEntry.AuthorityGetWidget()->GetPlayerComponent();
 		if (PreviousOwner != nullptr && PreviousOwner != this)
 		{
 			FFrame::KismetExecutionMessage(TEXT("The widget was created for another player. It can't be added."), ELogVerbosity::Warning, "InvalidPlayerParent");
 		}
 		else
 		{
-			InEntry.SetWidget(InEntry.GetWidget()); // to make sure the id is set
-			InEntry.GetWidget()->AuthoritySetParent(this, FUIFrameworkParentWidget(this));
+			InEntry.AuthoritySetWidget(InEntry.AuthorityGetWidget()); // to make sure the id is set
+			InEntry.AuthorityGetWidget()->AuthoritySetParent(this, FUIFrameworkParentWidget(this));
 			RootList.AddEntry(InEntry);
 		}
 	}
@@ -148,7 +147,7 @@ void UUIFrameworkPlayerComponent::RemoveWidget(UUIFrameworkWidget* Widget)
 		else
 		{
 			RootList.RemoveEntry(Widget);
-			Widget->AuthoritySetParent(nullptr, FUIFrameworkParentWidget());
+			Widget->AuthoritySetParent(this, FUIFrameworkParentWidget());
 		}
 	}
 }
@@ -181,13 +180,13 @@ void UUIFrameworkPlayerComponent::TickComponent(float DeltaTime, enum ELevelTick
 					}
 					else
 					{
-						Entry->Parent->LocalAddChild(Entry->Child);
+						Entry->Parent->LocalAddChild(Entry->ChildId);
 					}
 				}
 			}
 			else
 			{
-				ensureMsgf(false, TEXT("A widget would was added but couldn't be found anymore."));
+				UE_LOG(LogUIFramework, Verbose, TEXT("A widget would was added but couldn't be found anymore."));
 			}
 		}
 
@@ -199,7 +198,7 @@ void UUIFrameworkPlayerComponent::TickComponent(float DeltaTime, enum ELevelTick
 	}
 }
 
-void UUIFrameworkPlayerComponent::LocalAddWidget(const FUIFrameworkWidgetTreeEntry& Entry)
+void UUIFrameworkPlayerComponent::LocalWidgetWasAddedToTree(const FUIFrameworkWidgetTreeEntry& Entry)
 {
 	APlayerController* LocalOwner = GetPlayerController();
 	check(!LocalOwner->HasAuthority());
@@ -260,7 +259,7 @@ void UUIFrameworkPlayerComponent::LocalAddWidget(const FUIFrameworkWidgetTreeEnt
 	PrimaryComponentTick.SetTickFunctionEnable(NetReplicationPending.Num() > 0 || AddPending.Num() > 0 || ClassesToLoad.Num() > 0);
 }
 
-void UUIFrameworkPlayerComponent::LocalRemoveWidget(const FUIFrameworkWidgetTreeEntry& Entry)
+void UUIFrameworkPlayerComponent::LocalWidgetRemovedFromTree(const FUIFrameworkWidgetTreeEntry& Entry)
 {
 	APlayerController* LocalOwner = GetPlayerController();
 	check(!LocalOwner->HasAuthority());
@@ -315,23 +314,49 @@ void UUIFrameworkPlayerComponent::LocalOnClassLoaded(TSoftClassPtr<UWidget> Widg
 
 void UUIFrameworkPlayerComponent::LocalAddChild(FUIFrameworkWidgetId WidgetId)
 {
+	UUIFrameworkWidget* Widget = GetWidgetTree().FindWidgetById(WidgetId);
 	if (FUIFrameworkGameLayerSlot* LayerEntry = RootList.FindEntry(WidgetId))
 	{
-		if (LayerEntry->GetWidget() && LayerEntry->GetWidget()->LocalGetUMGWidget())
+		if (Widget)
 		{
-			LayerEntry->GetWidget()->LocalGetUMGWidget()->RemoveFromParent();
-			if (UUserWidget* UserWidget = Cast<UUserWidget>(LayerEntry->GetWidget()->LocalGetUMGWidget()))
+			UWidget* UMGWidget = Widget->LocalGetUMGWidget();
+			if (ensure(UMGWidget))
 			{
-				LayerEntry->LocalPreviousWidgetId = LayerEntry->GetWidgetId();
-				if (LayerEntry->Type == EUIFrameworkGameLayerType::Viewport)
+				UMGWidget->RemoveFromParent();
+				if (UGameViewportSubsystem* Subsystem = UGameViewportSubsystem::Get(GetWorld()))
 				{
-					UserWidget->AddToViewport(LayerEntry->ZOrder);
-				}
-				else
-				{
-					UserWidget->AddToPlayerScreen(LayerEntry->ZOrder);
+					LayerEntry->LocalAquireWidget();
+
+					FGameViewportWidgetSlot GameViewportWidgetSlot;
+					GameViewportWidgetSlot.ZOrder = LayerEntry->ZOrder;
+					if (LayerEntry->Type == EUIFrameworkGameLayerType::Viewport)
+					{
+						Subsystem->AddWidget(UMGWidget, GameViewportWidgetSlot);
+					}
+					else
+					{
+						APlayerController* LocalOwner = GetPlayerController();
+						check(LocalOwner);
+						Subsystem->AddWidgetForPlayer(UMGWidget, LocalOwner->GetLocalPlayer(), GameViewportWidgetSlot);
+					}
 				}
 			}
 		}
+		else
+		{
+			UE_LOG(LogUIFramework, Log, TEXT("The widget '%" INT64_FMT "' doesn't exist in the WidgetTree."), WidgetId.GetKey());
+			if (Widget && Widget->LocalGetUMGWidget())
+			{
+				Widget->LocalGetUMGWidget()->RemoveFromParent();
+			}
+		}
+	}
+	else
+	{
+		if (Widget && Widget->LocalGetUMGWidget())
+		{
+			Widget->LocalGetUMGWidget()->RemoveFromParent();
+		}
+		UE_LOG(LogUIFramework, Verbose, TEXT("The widget '%" INT64_FMT "' was not found in the RootList."), WidgetId.GetKey());
 	}
 }
diff --git a/Engine/Plugins/Experimental/UIFramework/Source/Private/UIFWidget.cpp b/Engine/Plugins/Experimental/UIFramework/Source/Private/UIFWidget.cpp
--- a/Engine/Plugins/Experimental/UIFramework/Source/Private/UIFWidget.cpp
+++ b/Engine/Plugins/Experimental/UIFramework/Source/Private/UIFWidget.cpp
@@ -66,19 +66,6 @@ void UUIFrameworkWidget::GetLifetimeReplicatedProps(TArray<FLifetimeProperty>& O
 	DOREPLIFETIME_WITH_PARAMS_FAST(ThisClass, Id, Params);
 }
 
-void UUIFrameworkWidget::BeginDestroy()
-{
-	Super::BeginDestroy();
-
-	if (UObjectInitialized() && LocalGetUMGWidget() != nullptr)
-	{
-		//if (UUIFrameworkManagerSubsystem* Subsystem = UUIFrameworkManagerSubsystem::Get(GetOuterAPlayerController()))
-		//{
-		//	Subsystem->LocalRemoveFromParent(GetLocalWidget());
-		//}
-	}
-}
-
 void UUIFrameworkWidget::AuthoritySetParent(UUIFrameworkPlayerComponent* NewOwner, FUIFrameworkParentWidget NewParent)
 {
 	const bool bDifferentOwner = NewOwner != OwnerPlayerComponent;
@@ -140,11 +127,19 @@ void UUIFrameworkWidget::SetParentPlayerOwnerRecursive()
 		});
 }
 
-void UUIFrameworkWidget::LocalAddChild(UUIFrameworkWidget* Child)
+void UUIFrameworkWidget::LocalAddChild(FUIFrameworkWidgetId ChildId)
 {
-	if (LocalUMGWidget)
+	// By default we should remove the widget from its previous parent.
+	//Adding a widget to a new slot will automaticly remove it from its previous parent.
+	if (OwnerPlayerComponent)
 	{
-		LocalUMGWidget->RemoveFromParent();
+		if (UUIFrameworkWidget* Widget = OwnerPlayerComponent->GetWidgetTree().FindWidgetById(ChildId))
+		{
+			if (UWidget* UMGWidget = Widget->LocalGetUMGWidget())
+			{
+				UMGWidget->RemoveFromParent();
+			}
+		}
 	}
 }
 
diff --git a/Engine/Plugins/Experimental/UIFramework/Source/Private/Widgets/UIFButton.cpp b/Engine/Plugins/Experimental/UIFramework/Source/Private/Widgets/UIFButton.cpp
--- a/Engine/Plugins/Experimental/UIFramework/Source/Private/Widgets/UIFButton.cpp
+++ b/Engine/Plugins/Experimental/UIFramework/Source/Private/Widgets/UIFButton.cpp
@@ -1,6 +1,8 @@
 // Copyright Epic Games, Inc. All Rights Reserved.
 
 #include "Widgets/UIFButton.h"
+#include "Types/UIFWidgetTree.h"
+#include "UIFLog.h"
 #include "UIFPlayerComponent.h"
 
 #include "Components/Button.h"
@@ -33,31 +35,32 @@ void UUIFrameworkButton::GetLifetimeReplicatedProps(TArray<FLifetimeProperty>& O
 
 void UUIFrameworkButton::SetContent(FUIFrameworkSimpleSlot InEntry)
 {
-	bool bWidgetIsDifferent = Slot.GetWidget() != InEntry.GetWidget();
+	bool bWidgetIsDifferent = Slot.AuthorityGetWidget() != InEntry.AuthorityGetWidget();
 	if (bWidgetIsDifferent)
 	{
-		if (Slot.GetWidget())
+		// Remove previous widget
+		if (Slot.AuthorityGetWidget())
 		{
-			Slot.GetWidget()->AuthoritySetParent(nullptr, FUIFrameworkParentWidget());
+			Slot.AuthorityGetWidget()->AuthoritySetParent(GetPlayerComponent(), FUIFrameworkParentWidget());
 		}
 
-		if (InEntry.GetWidget())
+		if (InEntry.AuthorityGetWidget())
 		{
-			UUIFrameworkPlayerComponent* PreviousOwner = InEntry.GetWidget()->GetPlayerComponent();
+			UUIFrameworkPlayerComponent* PreviousOwner = InEntry.AuthorityGetWidget()->GetPlayerComponent();
 			if (PreviousOwner != nullptr && PreviousOwner != GetPlayerComponent())
 			{
-				Slot.SetWidget(nullptr);
+				Slot.AuthoritySetWidget(nullptr);
 				FFrame::KismetExecutionMessage(TEXT("The widget was created for another player. It can't be added."), ELogVerbosity::Warning, "InvalidPlayerParent");
 			}
 		}
 	}
 
 	Slot = InEntry;
-	Slot.SetWidget(InEntry.GetWidget()); // to make sure the id is set
+	Slot.AuthoritySetWidget(InEntry.AuthorityGetWidget()); // to make sure the id is set
 
-	if (bWidgetIsDifferent && Slot.GetWidget())
+	if (bWidgetIsDifferent && Slot.AuthorityGetWidget())
 	{
-		Slot.GetWidget()->AuthoritySetParent(GetPlayerComponent(), FUIFrameworkParentWidget(this));
+		Slot.AuthorityGetWidget()->AuthoritySetParent(GetPlayerComponent(), FUIFrameworkParentWidget(this));
 	}
 
 	MARK_PROPERTY_DIRTY_FROM_NAME(ThisClass, Slot, this);
@@ -67,15 +70,15 @@ void UUIFrameworkButton::SetContent(FUIFrameworkSimpleSlot InEntry)
 void UUIFrameworkButton::AuthorityForEachChildren(const TFunctionRef<void(UUIFrameworkWidget*)>& Func)
 {
 	Super::AuthorityForEachChildren(Func);
-	Func(Slot.GetWidget());
+	Func(Slot.AuthorityGetWidget());
 }
 
 void UUIFrameworkButton::AuthorityRemoveChild(UUIFrameworkWidget* Widget)
 {
 	Super::AuthorityRemoveChild(Widget);
-	ensure(Widget == Slot.GetWidget());
+	ensure(Widget == Slot.AuthorityGetWidget());
 
-	Slot.SetWidget(nullptr);;
+	Slot.AuthoritySetWidget(nullptr);;
 	MARK_PROPERTY_DIRTY_FROM_NAME(ThisClass, Slot, this);
 }
 
@@ -87,19 +90,32 @@ void UUIFrameworkButton::LocalOnUMGWidgetCreated()
 }
 
 
-void UUIFrameworkButton::LocalAddChild(UUIFrameworkWidget* Child)
+void UUIFrameworkButton::LocalAddChild(FUIFrameworkWidgetId ChildId)
 {
-	Super::LocalAddChild(Child);
-	if (ensure(Child && Child->GetWidgetId() == Slot.GetWidgetId()))
+	if (ChildId == Slot.GetWidgetId())
 	{
-		UWidget* ChildWidget = Child->LocalGetUMGWidget();
-		if (ensure(ChildWidget))
+		check(GetPlayerComponent());
+		if (UUIFrameworkWidget* ChildWidget = GetPlayerComponent()->GetWidgetTree().FindWidgetById(ChildId))
 		{
-			UButton* Button = CastChecked<UButton>(LocalGetUMGWidget());
-			Button->AddChild(ChildWidget);
-			Slot.LocalPreviousWidgetId = Slot.GetWidgetId();
+			UWidget* ChildUMGWidget = ChildWidget->LocalGetUMGWidget();
+			if (ensure(ChildUMGWidget))
+			{
+				Slot.LocalAquireWidget();
+
+				UButton* Button = CastChecked<UButton>(LocalGetUMGWidget());
+				Button->ClearChildren();
+				UButtonSlot* ButtonSlot = CastChecked<UButtonSlot>(Button->AddChild(ChildUMGWidget));
+				ButtonSlot->SetPadding(Slot.Padding);
+				ButtonSlot->SetHorizontalAlignment(Slot.HorizontalAlignment);
+				ButtonSlot->SetVerticalAlignment(Slot.VerticalAlignment);
+			}
 		}
 	}
+	else
+	{
+		UE_LOG(LogUIFramework, Verbose, TEXT("The widget '%" INT64_FMT "' was not found in the Button Slots."), ChildId.GetKey());
+		Super::LocalAddChild(ChildId);
+	}
 }
 
 void UUIFrameworkButton::ServerClick_Implementation()
@@ -115,12 +131,14 @@ void UUIFrameworkButton::ServerClick_Implementation()
 
 void UUIFrameworkButton::OnRep_Slot()
 {
-	if (LocalGetUMGWidget() && Slot.GetWidget() && Slot.GetWidgetId() == Slot.LocalPreviousWidgetId)
+	if (LocalGetUMGWidget() && Slot.LocalIsAquiredWidgetValid())
 	{
-		UButtonSlot* ButtonSlot = CastChecked<UButtonSlot>(LocalGetUMGWidget()->Slot);
-		ButtonSlot->SetPadding(Slot.Padding);
-		ButtonSlot->SetHorizontalAlignment(Slot.HorizontalAlignment);
-		ButtonSlot->SetVerticalAlignment(Slot.VerticalAlignment);
+		if (UButtonSlot* ButtonSlot = Cast<UButtonSlot>(LocalGetUMGWidget()->Slot))
+		{
+			ButtonSlot->SetPadding(Slot.Padding);
+			ButtonSlot->SetHorizontalAlignment(Slot.HorizontalAlignment);
+			ButtonSlot->SetVerticalAlignment(Slot.VerticalAlignment);
+		}
 	}
 	// else do not do anything, the slot was not added yet or it was modified but was not applied yet by the PlayerComponent
 }
diff --git a/Engine/Plugins/Experimental/UIFramework/Source/Private/Widgets/UIFCanvasBox.cpp b/Engine/Plugins/Experimental/UIFramework/Source/Private/Widgets/UIFCanvasBox.cpp
new file mode 100644
--- /dev/null
+++ b/Engine/Plugins/Experimental/UIFramework/Source/Private/Widgets/UIFCanvasBox.cpp
@@ -0,0 +1,199 @@
+// Copyright Epic Games, Inc. All Rights Reserved.
+
+#include "Widgets/UIFCanvasBox.h"
+#include "Types/UIFWidgetTree.h"
+#include "UIFLog.h"
+#include "UIFPlayerComponent.h"
+
+#include "Components/CanvasPanel.h"
+#include "Components/CanvasPanelSlot.h"
+
+#include "Engine/ActorChannel.h"
+#include "Engine/Engine.h"
+#include "Engine/NetDriver.h"
+#include "GameFramework/Actor.h"
+#include "GameFramework/PlayerController.h"
+#include "Net/UnrealNetwork.h"
+
+
+/**
+ *
+ */
+void FUIFrameworkCanvasBoxSlotList::PostReplicatedChange(const TArrayView<int32>& ChangedIndices, int32 FinalSize)
+{
+	check(Owner);
+	for (int32 Index : ChangedIndices)
+	{
+		FUIFrameworkCanvasBoxSlot& Slot = Slots[Index];
+		if (Slot.LocalIsAquiredWidgetValid())
+		{
+			// Remove and add the widget again... 
+			// that may not work if they are on top of each other... The order may matter if the zorder is the same :( 
+			Owner->LocalAddChild(Slot.GetWidgetId());
+		}
+		//else it will be remove and the new widget will be added by the WidgetTree replication.
+	}
+}
+
+bool FUIFrameworkCanvasBoxSlotList::NetDeltaSerialize(FNetDeltaSerializeInfo& DeltaParms)
+{
+	return FFastArraySerializer::FastArrayDeltaSerialize<FUIFrameworkCanvasBoxSlot, FUIFrameworkCanvasBoxSlotList>(Slots, DeltaParms, *this);
+}
+
+void FUIFrameworkCanvasBoxSlotList::AddEntry(FUIFrameworkCanvasBoxSlot Entry)
+{
+	FUIFrameworkCanvasBoxSlot& NewEntry = Slots.Add_GetRef(MoveTemp(Entry));
+	MarkItemDirty(NewEntry);
+}
+
+bool FUIFrameworkCanvasBoxSlotList::RemoveEntry(UUIFrameworkWidget* Widget)
+{
+	check(Widget);
+	const int32 Index = Slots.IndexOfByPredicate([Widget](const FUIFrameworkCanvasBoxSlot& Entry) { return Entry.AuthorityGetWidget() == Widget; });
+	if (Index != INDEX_NONE)
+	{
+		Slots.RemoveAt(Index);
+		MarkArrayDirty();
+	}
+	return Index != INDEX_NONE;
+}
+
+FUIFrameworkCanvasBoxSlot* FUIFrameworkCanvasBoxSlotList::FindEntry(FUIFrameworkWidgetId WidgetId)
+{
+	return Slots.FindByPredicate([WidgetId](const FUIFrameworkCanvasBoxSlot& Entry) { return Entry.GetWidgetId() == WidgetId; });
+}
+
+void FUIFrameworkCanvasBoxSlotList::ForEachChildren(const TFunctionRef<void(UUIFrameworkWidget*)>& Func)
+{
+	for (FUIFrameworkCanvasBoxSlot& Slot : Slots)
+	{
+		Func(Slot.AuthorityGetWidget());
+	}
+}
+
+/**
+ *
+ */
+UUIFrameworkCanvasBox::UUIFrameworkCanvasBox()
+	: ReplicatedSlotList(this)
+{
+	WidgetClass = UCanvasPanel::StaticClass();
+}
+
+void UUIFrameworkCanvasBox::GetLifetimeReplicatedProps(TArray<FLifetimeProperty>& OutLifetimeProps) const
+{
+	Super::GetLifetimeReplicatedProps(OutLifetimeProps);
+
+	FDoRepLifetimeParams Params;
+	Params.bIsPushBased = true;
+	DOREPLIFETIME_WITH_PARAMS_FAST(ThisClass, ReplicatedSlotList, Params);
+}
+
+void UUIFrameworkCanvasBox::AddWidget(FUIFrameworkCanvasBoxSlot InEntry)
+{
+	if (InEntry.AuthorityGetWidget() == nullptr)
+	{
+		FFrame::KismetExecutionMessage(TEXT("The widget is invalid. It can't be added."), ELogVerbosity::Warning, "InvalidWidgetToAdd");
+	}
+	else if (GetPlayerComponent() && GetPlayerComponent() != InEntry.AuthorityGetWidget()->GetPlayerComponent())
+	{
+		check(GetPlayerComponent()->GetOwner()->HasAuthority());
+		FFrame::KismetExecutionMessage(TEXT("The widget was created for another player. It can't be removed on this player."), ELogVerbosity::Warning, "InvalidPlayerParentOnRemovedWidget");
+	}
+	else
+	{
+		InEntry.AuthoritySetWidget(InEntry.AuthorityGetWidget()); // to make sure the id is set
+		InEntry.AuthorityGetWidget()->AuthoritySetParent(GetPlayerComponent(), FUIFrameworkParentWidget(this));
+		AddEntry(InEntry);
+	}
+}
+
+void UUIFrameworkCanvasBox::RemoveWidget(UUIFrameworkWidget* Widget)
+{
+	if (Widget == nullptr)
+	{
+		FFrame::KismetExecutionMessage(TEXT("The widget is invalid. It can't be removed."), ELogVerbosity::Warning, "InvalidWidgetToRemove");
+	}
+	else
+	{
+		if (GetPlayerComponent() && GetPlayerComponent() != Widget->GetPlayerComponent())
+		{
+			check(GetPlayerComponent()->GetOwner()->HasAuthority());
+			FFrame::KismetExecutionMessage(TEXT("The widget was created for another player. It can't be removed on this player."), ELogVerbosity::Warning, "InvalidPlayerParentOnRemovedWidget");
+		}
+		else
+		{
+			RemoveEntry(Widget);
+			Widget->AuthoritySetParent(GetPlayerComponent(), FUIFrameworkParentWidget());
+		}
+	}
+}
+
+void UUIFrameworkCanvasBox::AuthorityForEachChildren(const TFunctionRef<void(UUIFrameworkWidget*)>& Func)
+{
+	Super::AuthorityForEachChildren(Func);
+
+	ReplicatedSlotList.ForEachChildren(Func);
+}
+
+void UUIFrameworkCanvasBox::AuthorityRemoveChild(UUIFrameworkWidget* Widget)
+{
+	Super::AuthorityRemoveChild(Widget);
+	RemoveEntry(Widget);
+}
+
+void UUIFrameworkCanvasBox::LocalAddChild(FUIFrameworkWidgetId ChildId)
+{
+	bool bIsAdded = false;
+	if (FUIFrameworkCanvasBoxSlot* CanvasEntry = FindEntry(ChildId))
+	{
+		check(GetPlayerComponent());
+		if (UUIFrameworkWidget* ChildWidget = GetPlayerComponent()->GetWidgetTree().FindWidgetById(ChildId))
+		{
+			UWidget* ChildUMGWidget = ChildWidget->LocalGetUMGWidget();
+			if (ensure(ChildUMGWidget))
+			{
+				UPanelSlot* PanelSlot = CastChecked<UCanvasPanel>(LocalGetUMGWidget())->AddChild(ChildUMGWidget);
+				checkf(PanelSlot, TEXT("CanvasPanel should be able to receive slot"));
+
+				CanvasEntry->LocalAquireWidget();
+
+				UCanvasPanelSlot* CanvasSlot = CastChecked<UCanvasPanelSlot>(PanelSlot);
+				{
+					FAnchorData AnchorData;
+					AnchorData.Anchors = CanvasEntry->Anchors;
+					AnchorData.Offsets = CanvasEntry->Offsets;
+					AnchorData.Alignment = CanvasEntry->Alignment;
+					CanvasSlot->SetLayout(AnchorData);
+				}
+				CanvasSlot->SetZOrder(CanvasEntry->ZOrder);
+				CanvasSlot->SetAutoSize(CanvasEntry->bSizeToContent);
+			}
+		}
+		else
+		{
+			UE_LOG(LogUIFramework, Log, TEXT("The widget '%" INT64_FMT "' doesn't exist in the WidgetTree."), ChildId.GetKey());
+			Super::LocalAddChild(ChildId);
+		}
+	}
+	else
+	{
+		UE_LOG(LogUIFramework, Verbose, TEXT("The widget '%" INT64_FMT "' was not found in the Canvas Slots."), ChildId.GetKey());
+		Super::LocalAddChild(ChildId);
+	}
+}
+
+void UUIFrameworkCanvasBox::AddEntry(FUIFrameworkCanvasBoxSlot Entry)
+{
+	ReplicatedSlotList.AddEntry(Entry);
+}
+
+bool UUIFrameworkCanvasBox::RemoveEntry(UUIFrameworkWidget* Widget)
+{
+	return ReplicatedSlotList.RemoveEntry(Widget);
+}
+
+FUIFrameworkCanvasBoxSlot* UUIFrameworkCanvasBox::FindEntry(FUIFrameworkWidgetId WidgetId)
+{
+	return ReplicatedSlotList.FindEntry(WidgetId);
+}
\ No newline at end of file
diff --git a/Engine/Plugins/Experimental/UIFramework/Source/Public/Types/UIFSlotBase.h b/Engine/Plugins/Experimental/UIFramework/Source/Public/Types/UIFSlotBase.h
--- a/Engine/Plugins/Experimental/UIFramework/Source/Public/Types/UIFSlotBase.h
+++ b/Engine/Plugins/Experimental/UIFramework/Source/Public/Types/UIFSlotBase.h
@@ -24,31 +24,41 @@ struct UIFRAMEWORK_API FUIFrameworkSlotBase : public FFastArraySerializerItem
 
 	FUIFrameworkSlotBase() = default;
 
-	UUIFrameworkWidget* GetWidget() const
+	UUIFrameworkWidget* AuthorityGetWidget() const
 	{
 		return Widget;
 	}
 
-	void SetWidget(UUIFrameworkWidget* Widget);
+	void AuthoritySetWidget(UUIFrameworkWidget* Widget);
 
 	FUIFrameworkWidgetId GetWidgetId() const
 	{
 		return WidgetId;
 	}
 
-	/**
-	 * The widget that was previously added on the local UMG Widget.
-	 * The server may have changed it but the "application" of that modification may be applied on the next frame by the PlayerComponent.
-	 */
-	UPROPERTY(NotReplicated)
-	FUIFrameworkWidgetId LocalPreviousWidgetId;
+	void LocalAquireWidget()
+	{
+		LocalPreviousWidgetId = WidgetId;
+	}
+
+	bool LocalIsAquiredWidgetValid() const
+	{
+		return LocalPreviousWidgetId == WidgetId;
+	}
 
 private:
-	UPROPERTY(BlueprintReadWrite, Category = "UI Framework", meta = (AllowPrivateAccess = "true"))
+	UPROPERTY(BlueprintReadWrite, Category = "UI Framework", NotReplicated, meta = (AllowPrivateAccess = "true"))
 	TObjectPtr<UUIFrameworkWidget> Widget = nullptr;
 
 	UPROPERTY()
 	FUIFrameworkWidgetId WidgetId;
+
+	/**
+	 * The widget that was previously added on the local UMG Widget.
+	 * The server may have changed it but the "application" of that modification may be applied on the next frame by the PlayerComponent.
+	 */
+	UPROPERTY(NotReplicated)
+	FUIFrameworkWidgetId LocalPreviousWidgetId;
 };
 
 
diff --git a/Engine/Plugins/Experimental/UIFramework/Source/Public/Types/UIFWidgetId.h b/Engine/Plugins/Experimental/UIFramework/Source/Public/Types/UIFWidgetId.h
--- a/Engine/Plugins/Experimental/UIFramework/Source/Public/Types/UIFWidgetId.h
+++ b/Engine/Plugins/Experimental/UIFramework/Source/Public/Types/UIFWidgetId.h
@@ -50,6 +50,11 @@ public:
 	{
 		return Key != Other.Key;
 	}
+	
+	FORCEINLINE friend uint32 GetTypeHash(const FUIFrameworkWidgetId& Value)
+	{
+		return GetTypeHash(Value.Key);
+	}
 
 private:
 	UPROPERTY()
diff --git a/Engine/Plugins/Experimental/UIFramework/Source/Public/Types/UIFWidgetTree.h b/Engine/Plugins/Experimental/UIFramework/Source/Public/Types/UIFWidgetTree.h
--- a/Engine/Plugins/Experimental/UIFramework/Source/Public/Types/UIFWidgetTree.h
+++ b/Engine/Plugins/Experimental/UIFramework/Source/Public/Types/UIFWidgetTree.h
@@ -90,6 +90,10 @@ public:
 	FUIFrameworkWidgetTreeEntry* GetEntryByReplicationId(int32 ReplicationId);
 	const FUIFrameworkWidgetTreeEntry* GetEntryByReplicationId(int32 ReplicationId) const;
 
+	/** Find the widget by its unique Id. The widget needs to be in the Tree. */
+	UUIFrameworkWidget* FindWidgetById(FUIFrameworkWidgetId WidgetId);
+	const UUIFrameworkWidget* FindWidgetById(FUIFrameworkWidgetId WidgetId) const;
+
 private:
 	void AddChildInternal(UUIFrameworkWidget* Parent, UUIFrameworkWidget* Child);
 	void AddChildRecursiveInternal(UUIFrameworkWidget* Widget);
@@ -98,6 +102,9 @@ private:
 	UPROPERTY()
 	TArray<FUIFrameworkWidgetTreeEntry> Entries;
 
+	UPROPERTY(NotReplicated)
+	TMap<FUIFrameworkWidgetId, TWeakObjectPtr<UUIFrameworkWidget>> WidgetByIdMap;
+
 	UPROPERTY(NotReplicated, Transient)
 	TObjectPtr<UUIFrameworkPlayerComponent> OwnerComponent = nullptr;
 };
diff --git a/Engine/Plugins/Experimental/UIFramework/Source/Public/UIFPlayerComponent.h b/Engine/Plugins/Experimental/UIFramework/Source/Public/UIFPlayerComponent.h
--- a/Engine/Plugins/Experimental/UIFramework/Source/Public/UIFPlayerComponent.h
+++ b/Engine/Plugins/Experimental/UIFramework/Source/Public/UIFPlayerComponent.h
@@ -135,9 +135,8 @@ public:
 	//~ End UActorComponent
 
 	void AuthorityRemoveChild(UUIFrameworkWidget* Widget);
-	void LocalAddWidget(const FUIFrameworkWidgetTreeEntry& Entry);
-	void LocalRemoveWidget(const FUIFrameworkWidgetTreeEntry& Entry);
-
+	void LocalWidgetWasAddedToTree(const FUIFrameworkWidgetTreeEntry& Entry);
+	void LocalWidgetRemovedFromTree(const FUIFrameworkWidgetTreeEntry& Entry);
 
 private:
 	void LocalOnClassLoaded(TSoftClassPtr<UWidget> WidgetClass);
diff --git a/Engine/Plugins/Experimental/UIFramework/Source/Public/UIFWidget.h b/Engine/Plugins/Experimental/UIFramework/Source/Public/UIFWidget.h
--- a/Engine/Plugins/Experimental/UIFramework/Source/Public/UIFWidget.h
+++ b/Engine/Plugins/Experimental/UIFramework/Source/Public/UIFWidget.h
@@ -26,7 +26,6 @@ public:
 	}
 	virtual int32 GetFunctionCallspace(UFunction* Function, FFrame* Stack) override;
 	virtual bool CallRemoteFunction(UFunction* Function, void* Parameters, FOutParmRec* OutParms, FFrame* Stack) override;
-	virtual void BeginDestroy() override;
 	//~ End UObject
 
 	FUIFrameworkWidgetId GetWidgetId() const
@@ -61,7 +60,7 @@ public:
 	}
 
 	void LocalCreateUMGWidget(UUIFrameworkPlayerComponent* Owner);
-	virtual void LocalAddChild(UUIFrameworkWidget* Child);
+	virtual void LocalAddChild(FUIFrameworkWidgetId ChildId);
 	void LocalDestroyUMGWidget();
 
 protected:
@@ -76,11 +75,11 @@ protected:
 	TSoftClassPtr<UWidget> WidgetClass; // todo: make this private and use a constructor argument
 
 private:
-	//~ Authority and Client
+	//~ Authority and Local
 	UPROPERTY(Replicated, Transient, DuplicateTransient)
 	FUIFrameworkWidgetId Id = FUIFrameworkWidgetId::MakeNew();
 
-	//~ Authority and Client
+	//~ Authority and Local
 	UPROPERTY(Transient)
 	TObjectPtr<UUIFrameworkPlayerComponent> OwnerPlayerComponent = nullptr;
 
diff --git a/Engine/Plugins/Experimental/UIFramework/Source/Public/Widgets/UIFButton.h b/Engine/Plugins/Experimental/UIFramework/Source/Public/Widgets/UIFButton.h
--- a/Engine/Plugins/Experimental/UIFramework/Source/Public/Widgets/UIFButton.h
+++ b/Engine/Plugins/Experimental/UIFramework/Source/Public/Widgets/UIFButton.h
@@ -8,11 +8,6 @@
 
 #include "UIFButton.generated.h"
 
-/**
- *
- */
-
-
 /**
  *
  */
@@ -40,7 +35,7 @@ public:
 
 	virtual void AuthorityForEachChildren(const TFunctionRef<void(UUIFrameworkWidget*)>& Func) override;
 	virtual void AuthorityRemoveChild(UUIFrameworkWidget* Widget) override;
-	virtual void LocalAddChild(UUIFrameworkWidget* Child) override;
+	virtual void LocalAddChild(FUIFrameworkWidgetId ChildId) override;
 
 protected:
 	virtual void LocalOnUMGWidgetCreated() override;
diff --git a/Engine/Plugins/Experimental/UIFramework/Source/Public/Widgets/UIFCanvasBox.h b/Engine/Plugins/Experimental/UIFramework/Source/Public/Widgets/UIFCanvasBox.h
new file mode 100644
--- /dev/null
+++ b/Engine/Plugins/Experimental/UIFramework/Source/Public/Widgets/UIFCanvasBox.h
@@ -0,0 +1,121 @@
+// Copyright Epic Games, Inc. All Rights Reserved.
+
+#pragma once
+
+#include "Layout/Margin.h"
+#include "Types/UIFSlotBase.h"
+#include "UIFWidget.h"
+#include "Widgets/Layout/Anchors.h"
+
+#include "UIFCanvasBox.generated.h"
+
+class UUIFrameworkCanvasBox;
+
+/**
+ *
+ */
+USTRUCT(BlueprintType)
+struct FUIFrameworkCanvasBoxSlot : public FUIFrameworkSlotBase
+{
+	GENERATED_BODY()
+
+	/** Anchors. */
+	UPROPERTY(EditAnywhere, BlueprintReadWrite, Category = Canvas)
+	FAnchors Anchors = FAnchors(0.0f, 0.0f);
+
+	/** Offset. */
+	UPROPERTY(EditAnywhere, BlueprintReadWrite, Category = Canvas)
+	FMargin Offsets = FMargin(0.f, 0.f, 0.f, 0.f);
+
+	/**
+	 * Alignment is the pivot point of the widget.  Starting in the upper left at (0,0),
+	 * ending in the lower right at (1,1).  Moving the alignment point allows you to move
+	 * the origin of the widget.
+	 */
+	UPROPERTY(EditAnywhere, BlueprintReadWrite, Category = Canvas)
+	FVector2D Alignment = FVector2D::ZeroVector;
+
+	/** The order priority this widget is rendered inside the layer. Higher values are rendered last (and so they will appear to be on top). */
+	UPROPERTY(EditAnywhere, BlueprintReadWrite, Category = Canvas)
+	int32 ZOrder = 0;
+
+	/** When true we use the widget's desired size */
+	UPROPERTY(EditAnywhere, BlueprintReadWrite, Category = Canvas)
+	bool bSizeToContent = false;
+};
+
+
+/**
+ *
+ */
+USTRUCT()
+struct UIFRAMEWORK_API FUIFrameworkCanvasBoxSlotList : public FFastArraySerializer
+{
+	GENERATED_BODY()
+
+	FUIFrameworkCanvasBoxSlotList() = default;
+	FUIFrameworkCanvasBoxSlotList(UUIFrameworkCanvasBox* InOwner)
+		: Owner(InOwner)
+	{}
+
+public:
+	//~ Begin of FFastArraySerializer
+	void PostReplicatedChange(const TArrayView<int32>& ChangedIndices, int32 FinalSize);
+	//~ End of FFastArraySerializer
+
+	bool NetDeltaSerialize(FNetDeltaSerializeInfo& DeltaParms);
+
+	void AddEntry(FUIFrameworkCanvasBoxSlot Entry);
+	bool RemoveEntry(UUIFrameworkWidget* Widget);
+	FUIFrameworkCanvasBoxSlot* FindEntry(FUIFrameworkWidgetId WidgetId);
+	void ForEachChildren(const TFunctionRef<void(UUIFrameworkWidget*)>& Func);
+
+private:
+	UPROPERTY()
+	TArray<FUIFrameworkCanvasBoxSlot> Slots;
+
+	UPROPERTY(NotReplicated, Transient)
+	TObjectPtr<UUIFrameworkCanvasBox> Owner;
+};
+
+
+template<>
+struct TStructOpsTypeTraits<FUIFrameworkCanvasBoxSlotList> : public TStructOpsTypeTraitsBase2<FUIFrameworkCanvasBoxSlotList>
+{
+	enum { WithNetDeltaSerializer = true };
+};
+
+
+/**
+ *
+ */
+UCLASS(DisplayName="Canvas UIFramework")
+class UIFRAMEWORK_API UUIFrameworkCanvasBox : public UUIFrameworkWidget
+{
+	GENERATED_BODY()
+
+	friend FUIFrameworkCanvasBoxSlotList;
+
+public:
+	UUIFrameworkCanvasBox();
+
+public:
+	UFUNCTION(BlueprintCallable, BlueprintAuthorityOnly, Category = "UI Framework")
+	void AddWidget(FUIFrameworkCanvasBoxSlot Widget);
+
+	UFUNCTION(BlueprintCallable, BlueprintAuthorityOnly, Category = "UI Framework")
+	void RemoveWidget(UUIFrameworkWidget* Widget);
+
+	virtual void AuthorityForEachChildren(const TFunctionRef<void(UUIFrameworkWidget*)>& Func) override;
+	virtual void AuthorityRemoveChild(UUIFrameworkWidget* Widget) override;
+	virtual void LocalAddChild(FUIFrameworkWidgetId ChildId) override;
+
+private:
+	void AddEntry(FUIFrameworkCanvasBoxSlot Entry);
+	bool RemoveEntry(UUIFrameworkWidget* Widget);
+	FUIFrameworkCanvasBoxSlot* FindEntry(FUIFrameworkWidgetId WidgetId);
+
+private:
+	UPROPERTY(Replicated)
+	FUIFrameworkCanvasBoxSlotList ReplicatedSlotList;
+};
