UIFramework: Add StackBox

diff --git a/Engine/Plugins/Experimental/UIFramework/Source/Private/Widgets/UIFStackBox.cpp b/Engine/Plugins/Experimental/UIFramework/Source/Private/Widgets/UIFStackBox.cpp
new file mode 100644
--- /dev/null
+++ b/Engine/Plugins/Experimental/UIFramework/Source/Private/Widgets/UIFStackBox.cpp
@@ -0,0 +1,229 @@
+// Copyright Epic Games, Inc. All Rights Reserved.
+
+#include "Widgets/UIFStackBox.h"
+#include "Types/UIFWidgetTree.h"
+#include "UIFLog.h"
+#include "UIFPlayerComponent.h"
+
+#include "Components/StackBox.h"
+#include "Components/StackBoxSlot.h"
+
+#include "Engine/ActorChannel.h"
+#include "Engine/Engine.h"
+#include "Engine/NetDriver.h"
+#include "GameFramework/Actor.h"
+#include "GameFramework/PlayerController.h"
+#include "Net/UnrealNetwork.h"
+
+
+/**
+ *
+ */
+void FUIFrameworkStackBoxSlotList::PostReplicatedChange(const TArrayView<int32>& ChangedIndices, int32 FinalSize)
+{
+	check(Owner);
+	for (int32 Index : ChangedIndices)
+	{
+		FUIFrameworkStackBoxSlot& Slot = Slots[Index];
+		if (Slot.LocalIsAquiredWidgetValid())
+		{
+			// Remove and add the widget again... 
+			// that may not work if they are on top of each other... The order may matter if the zorder is the same :( 
+			Owner->LocalAddChild(Slot.GetWidgetId());
+		}
+		//else it will be remove and the new widget will be added by the WidgetTree replication.
+	}
+}
+
+bool FUIFrameworkStackBoxSlotList::NetDeltaSerialize(FNetDeltaSerializeInfo& DeltaParms)
+{
+	return FFastArraySerializer::FastArrayDeltaSerialize<FUIFrameworkStackBoxSlot, FUIFrameworkStackBoxSlotList>(Slots, DeltaParms, *this);
+}
+
+void FUIFrameworkStackBoxSlotList::AddEntry(FUIFrameworkStackBoxSlot Entry)
+{
+	FUIFrameworkStackBoxSlot& NewEntry = Slots.Add_GetRef(MoveTemp(Entry));
+	MarkItemDirty(NewEntry);
+}
+
+bool FUIFrameworkStackBoxSlotList::RemoveEntry(UUIFrameworkWidget* Widget)
+{
+	check(Widget);
+	const int32 Index = Slots.IndexOfByPredicate([Widget](const FUIFrameworkStackBoxSlot& Entry) { return Entry.AuthorityGetWidget() == Widget; });
+	if (Index != INDEX_NONE)
+	{
+		Slots.RemoveAt(Index);
+		MarkArrayDirty();
+	}
+	return Index != INDEX_NONE;
+}
+
+FUIFrameworkStackBoxSlot* FUIFrameworkStackBoxSlotList::FindEntry(FUIFrameworkWidgetId WidgetId)
+{
+	return Slots.FindByPredicate([WidgetId](const FUIFrameworkStackBoxSlot& Entry) { return Entry.GetWidgetId() == WidgetId; });
+}
+
+void FUIFrameworkStackBoxSlotList::ForEachChildren(const TFunctionRef<void(UUIFrameworkWidget*)>& Func)
+{
+	for (FUIFrameworkStackBoxSlot& Slot : Slots)
+	{
+		Func(Slot.AuthorityGetWidget());
+	}
+}
+
+/**
+ *
+ */
+UUIFrameworkStackBox::UUIFrameworkStackBox()
+	: ReplicatedSlotList(this)
+{
+	WidgetClass = UStackBox::StaticClass();
+}
+
+void UUIFrameworkStackBox::GetLifetimeReplicatedProps(TArray<FLifetimeProperty>& OutLifetimeProps) const
+{
+	Super::GetLifetimeReplicatedProps(OutLifetimeProps);
+
+	FDoRepLifetimeParams Params;
+	Params.bIsPushBased = true;
+	DOREPLIFETIME_WITH_PARAMS_FAST(ThisClass, Orientation, Params);
+	DOREPLIFETIME_WITH_PARAMS_FAST(ThisClass, ReplicatedSlotList, Params);
+}
+
+void UUIFrameworkStackBox::AddWidget(FUIFrameworkStackBoxSlot InEntry)
+{
+	if (InEntry.AuthorityGetWidget() == nullptr)
+	{
+		FFrame::KismetExecutionMessage(TEXT("The widget is invalid. It can't be added."), ELogVerbosity::Warning, "InvalidWidgetToAdd");
+	}
+	else if (GetPlayerComponent() && GetPlayerComponent() != InEntry.AuthorityGetWidget()->GetPlayerComponent())
+	{
+		check(GetPlayerComponent()->GetOwner()->HasAuthority());
+		FFrame::KismetExecutionMessage(TEXT("The widget was created for another player. It can't be removed on this player."), ELogVerbosity::Warning, "InvalidPlayerParentOnRemovedWidget");
+	}
+	else
+	{
+		InEntry.AuthoritySetWidget(InEntry.AuthorityGetWidget()); // to make sure the id is set
+		InEntry.AuthorityGetWidget()->AuthoritySetParent(GetPlayerComponent(), FUIFrameworkParentWidget(this));
+		AddEntry(InEntry);
+	}
+}
+
+void UUIFrameworkStackBox::RemoveWidget(UUIFrameworkWidget* Widget)
+{
+	if (Widget == nullptr)
+	{
+		FFrame::KismetExecutionMessage(TEXT("The widget is invalid. It can't be removed."), ELogVerbosity::Warning, "InvalidWidgetToRemove");
+	}
+	else
+	{
+		if (GetPlayerComponent() && GetPlayerComponent() != Widget->GetPlayerComponent())
+		{
+			check(GetPlayerComponent()->GetOwner()->HasAuthority());
+			FFrame::KismetExecutionMessage(TEXT("The widget was created for another player. It can't be removed on this player."), ELogVerbosity::Warning, "InvalidPlayerParentOnRemovedWidget");
+		}
+		else
+		{
+			RemoveEntry(Widget);
+			Widget->AuthoritySetParent(GetPlayerComponent(), FUIFrameworkParentWidget());
+		}
+	}
+}
+
+EOrientation UUIFrameworkStackBox::GetOrientation() const
+{
+	return Orientation;
+}
+
+void UUIFrameworkStackBox::SetOrientation(EOrientation Value)
+{
+	if (Orientation != Value)
+	{
+		Orientation = Value;
+		if (UStackBox* StackBox = Cast<UStackBox>(LocalGetUMGWidget()))
+		{
+			StackBox->SetOrientation(Orientation);
+		}
+		MARK_PROPERTY_DIRTY_FROM_NAME(ThisClass, Orientation, this);
+	}
+}
+
+
+void UUIFrameworkStackBox::AuthorityForEachChildren(const TFunctionRef<void(UUIFrameworkWidget*)>& Func)
+{
+	Super::AuthorityForEachChildren(Func);
+
+	ReplicatedSlotList.ForEachChildren(Func);
+}
+
+void UUIFrameworkStackBox::AuthorityRemoveChild(UUIFrameworkWidget* Widget)
+{
+	Super::AuthorityRemoveChild(Widget);
+	RemoveEntry(Widget);
+}
+
+void UUIFrameworkStackBox::LocalOnUMGWidgetCreated()
+{
+	Super::LocalOnUMGWidgetCreated();
+	UStackBox* StackBox = CastChecked<UStackBox>(LocalGetUMGWidget());
+	StackBox->SetOrientation(Orientation);
+}
+
+void UUIFrameworkStackBox::LocalAddChild(FUIFrameworkWidgetId ChildId)
+{
+	bool bIsAdded = false;
+	if (FUIFrameworkStackBoxSlot* StackBoxEntry = FindEntry(ChildId))
+	{
+		check(GetPlayerComponent());
+		if (UUIFrameworkWidget* ChildWidget = GetPlayerComponent()->GetWidgetTree().FindWidgetById(ChildId))
+		{
+			UWidget* ChildUMGWidget = ChildWidget->LocalGetUMGWidget();
+			if (ensure(ChildUMGWidget))
+			{
+				UPanelSlot* PanelSlot = CastChecked<UStackBox>(LocalGetUMGWidget())->AddChild(ChildUMGWidget);
+				checkf(PanelSlot, TEXT("StackBoxPanel should be able to receive slot"));
+
+				StackBoxEntry->LocalAquireWidget();
+
+				UStackBoxSlot* StackBoxSlot = CastChecked<UStackBoxSlot>(PanelSlot);
+				StackBoxSlot->SetPadding(StackBoxEntry->Padding);
+				StackBoxSlot->SetSize(StackBoxEntry->Size);
+				StackBoxSlot->SetVerticalAlignment(StackBoxEntry->VerticalAlignment);
+				StackBoxSlot->SetHorizontalAlignment(StackBoxEntry->HorizontalAlignment);
+			}
+		}
+		else
+		{
+			UE_LOG(LogUIFramework, Log, TEXT("The widget '%" INT64_FMT "' doesn't exist in the WidgetTree."), ChildId.GetKey());
+			Super::LocalAddChild(ChildId);
+		}
+	}
+	else
+	{
+		UE_LOG(LogUIFramework, Verbose, TEXT("The widget '%" INT64_FMT "' was not found in the StackBox Slots."), ChildId.GetKey());
+		Super::LocalAddChild(ChildId);
+	}
+}
+
+void UUIFrameworkStackBox::OnRep_Orientation()
+{
+	if (UStackBox* StackBox = Cast<UStackBox>(LocalGetUMGWidget()))
+	{
+		StackBox->SetOrientation(Orientation);
+	}
+}
+
+void UUIFrameworkStackBox::AddEntry(FUIFrameworkStackBoxSlot Entry)
+{
+	ReplicatedSlotList.AddEntry(Entry);
+}
+
+bool UUIFrameworkStackBox::RemoveEntry(UUIFrameworkWidget* Widget)
+{
+	return ReplicatedSlotList.RemoveEntry(Widget);
+}
+
+FUIFrameworkStackBoxSlot* UUIFrameworkStackBox::FindEntry(FUIFrameworkWidgetId WidgetId)
+{
+	return ReplicatedSlotList.FindEntry(WidgetId);
+}
\ No newline at end of file
diff --git a/Engine/Plugins/Experimental/UIFramework/Source/Public/Widgets/UIFStackBox.h b/Engine/Plugins/Experimental/UIFramework/Source/Public/Widgets/UIFStackBox.h
new file mode 100644
--- /dev/null
+++ b/Engine/Plugins/Experimental/UIFramework/Source/Public/Widgets/UIFStackBox.h
@@ -0,0 +1,129 @@
+// Copyright Epic Games, Inc. All Rights Reserved.
+
+#pragma once
+
+#include "Components/SlateWrapperTypes.h"
+#include "Layout/Margin.h"
+#include "Types/SlateEnums.h"
+#include "Types/UIFSlotBase.h"
+#include "UIFWidget.h"
+#include "Widgets/Layout/Anchors.h"
+
+#include "UIFStackBox.generated.h"
+
+class UUIFrameworkStackBox;
+
+/**
+ *
+ */
+USTRUCT(BlueprintType)
+struct FUIFrameworkStackBoxSlot : public FUIFrameworkSlotBase
+{
+	GENERATED_BODY()
+
+	/** Horizontal alignment of the widget inside the slot. */
+	UPROPERTY(EditAnywhere, BlueprintReadWrite, Category = "UI Framework")
+	TEnumAsByte<EHorizontalAlignment> HorizontalAlignment = EHorizontalAlignment::HAlign_Fill;
+
+	/** Vertical alignment of the widget inside the slot. */
+	UPROPERTY(EditAnywhere, BlueprintReadWrite, Category = "UI Framework")
+	TEnumAsByte<EVerticalAlignment> VerticalAlignment = EVerticalAlignment::VAlign_Fill;
+
+	/** Distance between that surrounds the widget inside the slot. */
+	UPROPERTY(EditAnywhere, BlueprintReadWrite, Category = "UI Framework")
+	FMargin Padding = FMargin(0.f, 0.f, 0.f, 0.f);
+
+	/** How much space this slot should occupy in the direction of the panel. */
+	UPROPERTY(EditAnywhere, BlueprintReadWrite, Category = "UI Framework")
+	FSlateChildSize Size = FSlateChildSize(ESlateSizeRule::Automatic);
+};
+
+
+/**
+ *
+ */
+USTRUCT()
+struct UIFRAMEWORK_API FUIFrameworkStackBoxSlotList : public FFastArraySerializer
+{
+	GENERATED_BODY()
+
+	FUIFrameworkStackBoxSlotList() = default;
+	FUIFrameworkStackBoxSlotList(UUIFrameworkStackBox* InOwner)
+		: Owner(InOwner)
+	{}
+
+public:
+	//~ Begin of FFastArraySerializer
+	void PostReplicatedChange(const TArrayView<int32>& ChangedIndices, int32 FinalSize);
+	//~ End of FFastArraySerializer
+
+	bool NetDeltaSerialize(FNetDeltaSerializeInfo& DeltaParms);
+
+	void AddEntry(FUIFrameworkStackBoxSlot Entry);
+	bool RemoveEntry(UUIFrameworkWidget* Widget);
+	FUIFrameworkStackBoxSlot* FindEntry(FUIFrameworkWidgetId WidgetId);
+	void ForEachChildren(const TFunctionRef<void(UUIFrameworkWidget*)>& Func);
+
+private:
+	UPROPERTY()
+	TArray<FUIFrameworkStackBoxSlot> Slots;
+
+	UPROPERTY(NotReplicated, Transient)
+	TObjectPtr<UUIFrameworkStackBox> Owner;
+};
+
+
+template<>
+struct TStructOpsTypeTraits<FUIFrameworkStackBoxSlotList> : public TStructOpsTypeTraitsBase2<FUIFrameworkStackBoxSlotList>
+{
+	enum { WithNetDeltaSerializer = true };
+};
+
+
+/**
+ *
+ */
+UCLASS(DisplayName="StackBox UIFramework")
+class UIFRAMEWORK_API UUIFrameworkStackBox : public UUIFrameworkWidget
+{
+	GENERATED_BODY()
+
+	friend FUIFrameworkStackBoxSlotList;
+
+public:
+	UUIFrameworkStackBox();
+
+private:
+	/** The orientation of the stack box. */
+	UPROPERTY(ReplicatedUsing="OnRep_Orientation", EditAnywhere, BlueprintReadWrite, Setter, Getter, Category = "UI Framework", meta = (AllowPrivateAccess = "true"))
+	TEnumAsByte<EOrientation> Orientation = EOrientation::Orient_Horizontal;
+
+public:
+	UFUNCTION(BlueprintCallable, BlueprintAuthorityOnly, Category = "UI Framework")
+	void AddWidget(FUIFrameworkStackBoxSlot Widget);
+
+	UFUNCTION(BlueprintCallable, BlueprintAuthorityOnly, Category = "UI Framework")
+	void RemoveWidget(UUIFrameworkWidget* Widget);
+
+	/** Get the orientation of the stack box. */
+	EOrientation GetOrientation() const;
+	/** Set the orientation of the stack box. The existing elements will be rearranged. */
+	void SetOrientation(EOrientation Value);
+
+	virtual void AuthorityForEachChildren(const TFunctionRef<void(UUIFrameworkWidget*)>& Func) override;
+	virtual void AuthorityRemoveChild(UUIFrameworkWidget* Widget) override;
+	virtual void LocalAddChild(FUIFrameworkWidgetId ChildId) override;
+	virtual void LocalOnUMGWidgetCreated() override;
+
+private:
+	void AddEntry(FUIFrameworkStackBoxSlot Entry);
+	bool RemoveEntry(UUIFrameworkWidget* Widget);
+	FUIFrameworkStackBoxSlot* FindEntry(FUIFrameworkWidgetId WidgetId);
+
+	UFUNCTION()
+	void OnRep_Orientation();
+
+private:
+	UPROPERTY(Replicated)
+	FUIFrameworkStackBoxSlotList ReplicatedSlotList;
+};
