UIFramework: Fix module name

diff --git a/Engine/Plugins/Experimental/UIFramework/Source/Private/UIFrameworkModule.cpp b/Engine/Plugins/Experimental/UIFramework/Source/Private/UIFrameworkModule.cpp
--- a/Engine/Plugins/Experimental/UIFramework/Source/Private/UIFrameworkModule.cpp
+++ b/Engine/Plugins/Experimental/UIFramework/Source/Private/UIFrameworkModule.cpp
@@ -2,4 +2,4 @@
 
 #include "Modules/ModuleManager.h"
 	
-IMPLEMENT_MODULE(FDefaultModuleImpl, UISystem)
\ No newline at end of file
+IMPLEMENT_MODULE(FDefaultModuleImpl, UIFramework)
\ No newline at end of file
