UIFramework nullchecks

diff --git a/Engine/Plugins/Experimental/UIFramework/Source/Private/Types/UIFWidgetTree.cpp b/Engine/Plugins/Experimental/UIFramework/Source/Private/Types/UIFWidgetTree.cpp
--- a/Engine/Plugins/Experimental/UIFramework/Source/Private/Types/UIFWidgetTree.cpp
+++ b/Engine/Plugins/Experimental/UIFramework/Source/Private/Types/UIFWidgetTree.cpp
@@ -168,10 +168,13 @@ void FUIFrameworkWidgetTree::AddChildRecursiveInternal(UUIFrameworkWidget* Widge
 	FUIFrameworkWidgetTree* Self = this;
 	Widget->AuthorityForEachChildren([Self, Widget](UUIFrameworkWidget* ChildWidget)
 		{
-			FUIFrameworkWidgetTreeEntry& NewEntry = Self->Entries.Emplace_GetRef(Widget, ChildWidget);
-			Self->MarkItemDirty(NewEntry);
-			Self->WidgetByIdMap.FindOrAdd(ChildWidget->GetWidgetId()) = ChildWidget;
-			Self->AddChildRecursiveInternal(ChildWidget);
+			if (ChildWidget != nullptr)
+			{
+				FUIFrameworkWidgetTreeEntry& NewEntry = Self->Entries.Emplace_GetRef(Widget, ChildWidget);
+				Self->MarkItemDirty(NewEntry);
+				Self->WidgetByIdMap.FindOrAdd(ChildWidget->GetWidgetId()) = ChildWidget;
+				Self->AddChildRecursiveInternal(ChildWidget);
+			}
 		});
 }
 
diff --git a/Engine/Plugins/Experimental/UIFramework/Source/Private/UIFWidget.cpp b/Engine/Plugins/Experimental/UIFramework/Source/Private/UIFWidget.cpp
--- a/Engine/Plugins/Experimental/UIFramework/Source/Private/UIFWidget.cpp
+++ b/Engine/Plugins/Experimental/UIFramework/Source/Private/UIFWidget.cpp
@@ -120,10 +120,13 @@ void UUIFrameworkWidget::SetParentPlayerOwnerRecursive()
 	UUIFrameworkWidget* Self = this;
 	AuthorityForEachChildren([Self](UUIFrameworkWidget* Child)
 		{
-			check(Child->AuthorityGetParent().IsWidget() && Child->AuthorityGetParent().AsWidget() == Self);
-			Child->OwnerPlayerComponent = Self->OwnerPlayerComponent;
-			Child->AuthorityParent = FUIFrameworkParentWidget(Self);
-			Child->SetParentPlayerOwnerRecursive();
+			if (Child != nullptr)
+			{
+				check(Child->AuthorityGetParent().IsWidget() && Child->AuthorityGetParent().AsWidget() == Self);
+				Child->OwnerPlayerComponent = Self->OwnerPlayerComponent;
+				Child->AuthorityParent = FUIFrameworkParentWidget(Self);
+				Child->SetParentPlayerOwnerRecursive();
+			}
 		});
 }
 
diff --git a/Engine/Plugins/Experimental/UIFramework/Source/Private/Widgets/UIFButton.cpp b/Engine/Plugins/Experimental/UIFramework/Source/Private/Widgets/UIFButton.cpp
--- a/Engine/Plugins/Experimental/UIFramework/Source/Private/Widgets/UIFButton.cpp
+++ b/Engine/Plugins/Experimental/UIFramework/Source/Private/Widgets/UIFButton.cpp
@@ -70,7 +70,10 @@ void UUIFrameworkButton::SetContent(FUIFrameworkSimpleSlot InEntry)
 void UUIFrameworkButton::AuthorityForEachChildren(const TFunctionRef<void(UUIFrameworkWidget*)>& Func)
 {
 	Super::AuthorityForEachChildren(Func);
-	Func(Slot.AuthorityGetWidget());
+	if (UUIFrameworkWidget* ChildWidget = Slot.AuthorityGetWidget())
+	{
+		Func(ChildWidget);
+	}
 }
 
 void UUIFrameworkButton::AuthorityRemoveChild(UUIFrameworkWidget* Widget)
diff --git a/Engine/Plugins/Experimental/UIFramework/Source/Private/Widgets/UIFCanvasBox.cpp b/Engine/Plugins/Experimental/UIFramework/Source/Private/Widgets/UIFCanvasBox.cpp
--- a/Engine/Plugins/Experimental/UIFramework/Source/Private/Widgets/UIFCanvasBox.cpp
+++ b/Engine/Plugins/Experimental/UIFramework/Source/Private/Widgets/UIFCanvasBox.cpp
@@ -67,7 +67,10 @@ void FUIFrameworkCanvasBoxSlotList::ForEachChildren(const TFunctionRef<void(UUIF
 {
 	for (FUIFrameworkCanvasBoxSlot& Slot : Slots)
 	{
-		Func(Slot.AuthorityGetWidget());
+		if (UUIFrameworkWidget* ChildWidget = Slot.AuthorityGetWidget())
+		{
+			Func(ChildWidget);
+		}
 	}
 }
 
diff --git a/Engine/Plugins/Experimental/UIFramework/Source/Public/Widgets/UIFSimpleButton.h b/Engine/Plugins/Experimental/UIFramework/Source/Public/Widgets/UIFSimpleButton.h
--- a/Engine/Plugins/Experimental/UIFramework/Source/Public/Widgets/UIFSimpleButton.h
+++ b/Engine/Plugins/Experimental/UIFramework/Source/Public/Widgets/UIFSimpleButton.h
@@ -11,7 +11,7 @@
 /**
  *
  */
-DECLARE_DYNAMIC_MULTICAST_SPARSE_DELEGATE_OneParam(FUIFrameworkSimpleButtonClickEvent, UUIFrameworkSimpleButton, OnClick, FUIFrameworkClickEventArgument, Argument);
+DECLARE_MULTICAST_DELEGATE_OneParam(FUIFrameworkSimpleButtonClickEvent, FUIFrameworkClickEventArgument);
 
 
 /**
@@ -48,6 +48,5 @@ private:
 	FText Text;
 
 public:
-	UPROPERTY(BlueprintAssignable)
 	FUIFrameworkSimpleButtonClickEvent OnClick;
 };
