[UIFramework] Release Slate resources immediately when the local widget is destroyed.

diff --git a/Engine/Plugins/Experimental/UIFramework/Source/Private/UIFWidget.cpp b/Engine/Plugins/Experimental/UIFramework/Source/Private/UIFWidget.cpp
--- a/Engine/Plugins/Experimental/UIFramework/Source/Private/UIFWidget.cpp
+++ b/Engine/Plugins/Experimental/UIFramework/Source/Private/UIFWidget.cpp
@@ -169,6 +169,7 @@ void UUIFrameworkWidget::LocalDestroyUMGWidget()
 	if (LocalUMGWidget)
 	{
 		LocalUMGWidget->RemoveFromParent();
+		LocalUMGWidget->ReleaseSlateResources(true);
 	}
 	LocalUMGWidget = nullptr;
 	OwnerPlayerComponent = nullptr;
diff --git a/Engine/Plugins/Experimental/UIFramework/Source/Private/Widgets/UIFCanvasBox.cpp b/Engine/Plugins/Experimental/UIFramework/Source/Private/Widgets/UIFCanvasBox.cpp
--- a/Engine/Plugins/Experimental/UIFramework/Source/Private/Widgets/UIFCanvasBox.cpp
+++ b/Engine/Plugins/Experimental/UIFramework/Source/Private/Widgets/UIFCanvasBox.cpp
@@ -156,12 +156,10 @@ void UUIFrameworkCanvasBox::LocalAddChild(FUIFrameworkWidgetId ChildId)
 			UWidget* ChildUMGWidget = ChildWidget->LocalGetUMGWidget();
 			if (ensure(ChildUMGWidget))
 			{
-				UPanelSlot* PanelSlot = CastChecked<UCanvasPanel>(LocalGetUMGWidget())->AddChild(ChildUMGWidget);
-				checkf(PanelSlot, TEXT("CanvasPanel should be able to receive slot"));
+				UCanvasPanelSlot* CanvasSlot = CastChecked<UCanvasPanel>(LocalGetUMGWidget())->AddChildToCanvas(ChildUMGWidget);
+				checkf(CanvasSlot, TEXT("CanvasPanel should be able to receive slot"));
 
 				CanvasEntry->LocalAquireWidget();
-
-				UCanvasPanelSlot* CanvasSlot = CastChecked<UCanvasPanelSlot>(PanelSlot);
 				{
 					FAnchorData AnchorData;
 					AnchorData.Anchors = CanvasEntry->Anchors;
