Replicate UIFrameworkWidget with AddReplicatedSubObject.

diff --git a/Engine/Plugins/Experimental/UIFramework/Source/Private/Types/UIFWidgetTree.cpp b/Engine/Plugins/Experimental/UIFramework/Source/Private/Types/UIFWidgetTree.cpp
--- a/Engine/Plugins/Experimental/UIFramework/Source/Private/Types/UIFWidgetTree.cpp
+++ b/Engine/Plugins/Experimental/UIFramework/Source/Private/Types/UIFWidgetTree.cpp
@@ -1,6 +1,7 @@
 // Copyright Epic Games, Inc. All Rights Reserved.
 
 #include "Types/UIFWidgetTree.h"
+#include "UIFLog.h"
 #include "UIFPlayerComponent.h"
 #include "UIFWidget.h"
 
@@ -143,22 +144,29 @@ void FUIFrameworkWidgetTree::AddWidget(UUIFrameworkWidget* Parent, UUIFrameworkW
 
 void FUIFrameworkWidgetTree::AddChildInternal(UUIFrameworkWidget* Parent, UUIFrameworkWidget* Child)
 {
+	check(OwnerComponent);
 	int32 PreviousEntryIndex = Entries.IndexOfByPredicate([Child](const FUIFrameworkWidgetTreeEntry& Other){ return Other.Child == Child; });
 	if (PreviousEntryIndex != INDEX_NONE)
 	{
 		const FUIFrameworkWidgetTreeEntry& PreviousEntry = Entries[PreviousEntryIndex];
 		if (PreviousEntry.Parent != Parent)
 		{
+			// New Parent for an existing child.
 			Entries.RemoveAtSwap(PreviousEntryIndex);
 			FUIFrameworkWidgetTreeEntry& NewEntry = Entries.Emplace_GetRef(Parent, Child);
 			MarkItemDirty(NewEntry);
 		}
+		else
+		{
+			UE_LOG(LogUIFramework, Warning, TEXT("A widget was added to the WidgetTree but it's already added."));
+		}
 	}
 	else
 	{
 		FUIFrameworkWidgetTreeEntry& NewEntry = Entries.Emplace_GetRef(Parent, Child);
 		MarkItemDirty(NewEntry);
 		WidgetByIdMap.FindOrAdd(Child->GetWidgetId()) = Child;
+		OwnerComponent->GetOwner()->AddReplicatedSubObject(Child);
 		AddChildRecursiveInternal(Child);
 	}
 }
@@ -173,6 +181,7 @@ void FUIFrameworkWidgetTree::AddChildRecursiveInternal(UUIFrameworkWidget* Widge
 				FUIFrameworkWidgetTreeEntry& NewEntry = Self->Entries.Emplace_GetRef(Widget, ChildWidget);
 				Self->MarkItemDirty(NewEntry);
 				Self->WidgetByIdMap.FindOrAdd(ChildWidget->GetWidgetId()) = ChildWidget;
+				Self->OwnerComponent->GetOwner()->AddReplicatedSubObject(ChildWidget);
 				Self->AddChildRecursiveInternal(ChildWidget);
 			}
 		});
@@ -191,6 +200,7 @@ void FUIFrameworkWidgetTree::RemoveWidget(UUIFrameworkWidget* Widget)
 		{
 			UUIFrameworkWidget* CurrentWidgetToRemove = WidgetsToRemove.Last();
 			WidgetByIdMap.Remove(CurrentWidgetToRemove->GetWidgetId());
+			OwnerComponent->GetOwner()->RemoveReplicatedSubObject(CurrentWidgetToRemove);
 
 			WidgetsToRemove.RemoveAt(WidgetsToRemove.Num()-1);
 			for (int32 Index = Entries.Num() - 1; Index >= 0; --Index)
diff --git a/Engine/Plugins/Experimental/UIFramework/Source/Private/UIFPlayerComponent.cpp b/Engine/Plugins/Experimental/UIFramework/Source/Private/UIFPlayerComponent.cpp
--- a/Engine/Plugins/Experimental/UIFramework/Source/Private/UIFPlayerComponent.cpp
+++ b/Engine/Plugins/Experimental/UIFramework/Source/Private/UIFPlayerComponent.cpp
@@ -196,6 +196,25 @@ void UUIFrameworkPlayerComponent::TickComponent(float DeltaTime, enum ELevelTick
 
 		PrimaryComponentTick.SetTickFunctionEnable(false);
 	}
+	else
+	{
+		FMemMark Mark(FMemStack::Get());
+		TArray<int32, TMemStackAllocator<>> TempNetReplicationPending;
+		TempNetReplicationPending.Reserve(NetReplicationPending.Num());
+		for (int32 ReplicationId : NetReplicationPending)
+		{
+			TempNetReplicationPending.Add(ReplicationId);
+		}
+		NetReplicationPending.Empty(false);
+
+		for (int32 ReplicationId : TempNetReplicationPending)
+		{
+			if (const FUIFrameworkWidgetTreeEntry* Entry = WidgetTree.GetEntryByReplicationId(ReplicationId))
+			{
+				LocalWidgetWasAddedToTree(*Entry);
+			}
+		}
+	}
 }
 
 void UUIFrameworkPlayerComponent::LocalWidgetWasAddedToTree(const FUIFrameworkWidgetTreeEntry& Entry)
