Basic button now uses the button with slot.

diff --git a/Engine/Plugins/Experimental/UIFramework/Source/Public/Types/UIFEvents.h b/Engine/Plugins/Experimental/UIFramework/Source/Public/Types/UIFEvents.h
--- a/Engine/Plugins/Experimental/UIFramework/Source/Public/Types/UIFEvents.h
+++ b/Engine/Plugins/Experimental/UIFramework/Source/Public/Types/UIFEvents.h
@@ -26,3 +26,5 @@ struct FUIFrameworkClickEventArgument
 	UPROPERTY()
 	TObjectPtr<UUIFrameworkWidget> Sender;
 };
+
+DECLARE_MULTICAST_DELEGATE_OneParam(FUIFrameworkClickEvent, FUIFrameworkClickEventArgument);
diff --git a/Engine/Plugins/Experimental/UIFramework/Source/Public/Widgets/UIFButton.h b/Engine/Plugins/Experimental/UIFramework/Source/Public/Widgets/UIFButton.h
--- a/Engine/Plugins/Experimental/UIFramework/Source/Public/Widgets/UIFButton.h
+++ b/Engine/Plugins/Experimental/UIFramework/Source/Public/Widgets/UIFButton.h
@@ -8,11 +8,6 @@
 
 #include "UIFButton.generated.h"
 
-/**
- *
- */
-DECLARE_DYNAMIC_MULTICAST_SPARSE_DELEGATE_OneParam(FUIFrameworkButtonClickEvent, UUIFrameworkButton, OnClick, FUIFrameworkClickEventArgument, Argument);
-
 /**
  *
  */
@@ -48,8 +43,7 @@ private:
 	void OnRep_Slot();
 
 public:
-	UPROPERTY(BlueprintAssignable)
-	FUIFrameworkButtonClickEvent OnClick;
+	FUIFrameworkClickEvent OnClick;
 
 private:
 	UPROPERTY(/*ExposeOnSpawn, */ReplicatedUsing = OnRep_Slot)
diff --git a/Engine/Plugins/Experimental/UIFramework/Source/Public/Widgets/UIFSimpleButton.h b/Engine/Plugins/Experimental/UIFramework/Source/Public/Widgets/UIFSimpleButton.h
--- a/Engine/Plugins/Experimental/UIFramework/Source/Public/Widgets/UIFSimpleButton.h
+++ b/Engine/Plugins/Experimental/UIFramework/Source/Public/Widgets/UIFSimpleButton.h
@@ -7,13 +7,6 @@
 
 #include "UIFSimpleButton.generated.h"
 
-
-/**
- *
- */
-DECLARE_MULTICAST_DELEGATE_OneParam(FUIFrameworkSimpleButtonClickEvent, FUIFrameworkClickEventArgument);
-
-
 /**
  *
  */
@@ -48,5 +41,5 @@ private:
 	FText Text;
 
 public:
-	FUIFrameworkSimpleButtonClickEvent OnClick;
+	FUIFrameworkClickEvent OnClick;
 };
