UIFramework: Test if the owner uses the registered SubObject list before calling it.

diff --git a/Engine/Plugins/Experimental/UIFramework/Source/Private/Types/UIFWidgetTree.cpp b/Engine/Plugins/Experimental/UIFramework/Source/Private/Types/UIFWidgetTree.cpp
--- a/Engine/Plugins/Experimental/UIFramework/Source/Private/Types/UIFWidgetTree.cpp
+++ b/Engine/Plugins/Experimental/UIFramework/Source/Private/Types/UIFWidgetTree.cpp
@@ -166,7 +166,11 @@ void FUIFrameworkWidgetTree::AddChildInternal(UUIFrameworkWidget* Parent, UUIFra
 		FUIFrameworkWidgetTreeEntry& NewEntry = Entries.Emplace_GetRef(Parent, Child);
 		MarkItemDirty(NewEntry);
 		WidgetByIdMap.FindOrAdd(Child->GetWidgetId()) = Child;
-		OwnerComponent->GetOwner()->AddReplicatedSubObject(Child);
+
+		if (OwnerComponent->GetOwner()->IsUsingRegisteredSubObjectList())
+		{
+			OwnerComponent->GetOwner()->AddReplicatedSubObject(Child);
+		}
 		AddChildRecursiveInternal(Child);
 	}
 }
@@ -181,7 +185,10 @@ void FUIFrameworkWidgetTree::AddChildRecursiveInternal(UUIFrameworkWidget* Widge
 				FUIFrameworkWidgetTreeEntry& NewEntry = Self->Entries.Emplace_GetRef(Widget, ChildWidget);
 				Self->MarkItemDirty(NewEntry);
 				Self->WidgetByIdMap.FindOrAdd(ChildWidget->GetWidgetId()) = ChildWidget;
-				Self->OwnerComponent->GetOwner()->AddReplicatedSubObject(ChildWidget);
+				if (Self->OwnerComponent->GetOwner()->IsUsingRegisteredSubObjectList())
+				{
+					Self->OwnerComponent->GetOwner()->AddReplicatedSubObject(ChildWidget);
+				}
 				Self->AddChildRecursiveInternal(ChildWidget);
 			}
 		});
