UIFramework: Add a interface for the widget tree's owner

diff --git a/Engine/Plugins/Experimental/UIFramework/Source/Private/Types/UIFWidgetTree.cpp b/Engine/Plugins/Experimental/UIFramework/Source/Private/Types/UIFWidgetTree.cpp
--- a/Engine/Plugins/Experimental/UIFramework/Source/Private/Types/UIFWidgetTree.cpp
+++ b/Engine/Plugins/Experimental/UIFramework/Source/Private/Types/UIFWidgetTree.cpp
@@ -45,7 +45,10 @@ void FUIFrameworkWidgetTree::PreReplicatedRemove(const TArrayView<int32> Removed
 		FUIFrameworkWidgetTreeEntry& Entry = Entries[Index];
 		if (Entry.Child)
 		{
-			OwnerComponent->LocalWidgetRemovedFromTree(Entry);
+			if (ensure(Owner))
+			{
+				Owner->LocalWidgetRemovedFromTree(Entry);
+			}
 
 			// Was it remove but added by something else
 			{
@@ -78,7 +81,10 @@ void FUIFrameworkWidgetTree::PostReplicatedAdd(const TArrayView<int32> AddedIndi
 		FUIFrameworkWidgetTreeEntry& Entry = Entries[Index];
 		if (Entry.ParentId.IsValid() && Entry.ChildId.IsValid())
 		{
-			OwnerComponent->LocalWidgetWasAddedToTree(Entry);
+			if (ensure(Owner))
+			{
+				Owner->LocalWidgetWasAddedToTree(Entry);
+			}
 			WidgetByIdMap.FindOrAdd(Entry.ChildId) = Entry.Child;
 			if (!Entry.ParentId.IsRoot())
 			{
@@ -96,7 +102,10 @@ void FUIFrameworkWidgetTree::PostReplicatedChange(const TArrayView<int32>& Chang
 		FUIFrameworkWidgetTreeEntry& Entry = Entries[Index];
 		if (Entry.ParentId.IsValid() && Entry.ChildId.IsValid())
 		{
-			OwnerComponent->LocalWidgetWasAddedToTree(Entry);
+			if (ensure(Owner))
+			{
+				Owner->LocalWidgetWasAddedToTree(Entry);
+			}
 			WidgetByIdMap.FindOrAdd(Entry.ChildId) = Entry.Child;
 			if (!Entry.ParentId.IsRoot())
 			{
@@ -144,7 +153,6 @@ void FUIFrameworkWidgetTree::AddWidget(UUIFrameworkWidget* Parent, UUIFrameworkW
 
 void FUIFrameworkWidgetTree::AddChildInternal(UUIFrameworkWidget* Parent, UUIFrameworkWidget* Child)
 {
-	check(OwnerComponent);
 	int32 PreviousEntryIndex = Entries.IndexOfByPredicate([Child](const FUIFrameworkWidgetTreeEntry& Other){ return Other.Child == Child; });
 	if (PreviousEntryIndex != INDEX_NONE)
 	{
@@ -167,9 +175,9 @@ void FUIFrameworkWidgetTree::AddChildInternal(UUIFrameworkWidget* Parent, UUIFra
 		MarkItemDirty(NewEntry);
 		WidgetByIdMap.FindOrAdd(Child->GetWidgetId()) = Child;
 
-		if (OwnerComponent->GetOwner()->IsUsingRegisteredSubObjectList())
+		if (ensure(ReplicatedOwner) && ReplicatedOwner->IsUsingRegisteredSubObjectList())
 		{
-			OwnerComponent->GetOwner()->AddReplicatedSubObject(Child);
+			ReplicatedOwner->AddReplicatedSubObject(Child);
 		}
 		AddChildRecursiveInternal(Child);
 	}
@@ -185,9 +193,9 @@ void FUIFrameworkWidgetTree::AddChildRecursiveInternal(UUIFrameworkWidget* Widge
 				FUIFrameworkWidgetTreeEntry& NewEntry = Self->Entries.Emplace_GetRef(Widget, ChildWidget);
 				Self->MarkItemDirty(NewEntry);
 				Self->WidgetByIdMap.FindOrAdd(ChildWidget->GetWidgetId()) = ChildWidget;
-				if (Self->OwnerComponent->GetOwner()->IsUsingRegisteredSubObjectList())
+				if (ensure(Self->ReplicatedOwner) && Self->ReplicatedOwner->IsUsingRegisteredSubObjectList())
 				{
-					Self->OwnerComponent->GetOwner()->AddReplicatedSubObject(ChildWidget);
+					Self->ReplicatedOwner->AddReplicatedSubObject(ChildWidget);
 				}
 				Self->AddChildRecursiveInternal(ChildWidget);
 			}
@@ -207,7 +215,10 @@ void FUIFrameworkWidgetTree::RemoveWidget(UUIFrameworkWidget* Widget)
 		{
 			UUIFrameworkWidget* CurrentWidgetToRemove = WidgetsToRemove.Last();
 			WidgetByIdMap.Remove(CurrentWidgetToRemove->GetWidgetId());
-			OwnerComponent->GetOwner()->RemoveReplicatedSubObject(CurrentWidgetToRemove);
+			if (ensure(ReplicatedOwner) && ReplicatedOwner->IsUsingRegisteredSubObjectList())
+			{
+				ReplicatedOwner->RemoveReplicatedSubObject(CurrentWidgetToRemove);
+			}
 
 			WidgetsToRemove.RemoveAt(WidgetsToRemove.Num()-1);
 			for (int32 Index = Entries.Num() - 1; Index >= 0; --Index)
diff --git a/Engine/Plugins/Experimental/UIFramework/Source/Private/UIFPlayerComponent.cpp b/Engine/Plugins/Experimental/UIFramework/Source/Private/UIFPlayerComponent.cpp
--- a/Engine/Plugins/Experimental/UIFramework/Source/Private/UIFPlayerComponent.cpp
+++ b/Engine/Plugins/Experimental/UIFramework/Source/Private/UIFPlayerComponent.cpp
@@ -61,7 +61,7 @@ FUIFrameworkGameLayerSlot* FUIFrameworkGameLayerSlotList::FindEntry(FUIFramework
  */
 UUIFrameworkPlayerComponent::UUIFrameworkPlayerComponent()
 	: RootList(this)
-	, WidgetTree(this)
+	, WidgetTree(GetOwner(), this)
 {
 	SetIsReplicatedByDefault(true);
 	bWantsInitializeComponent = true;
diff --git a/Engine/Plugins/Experimental/UIFramework/Source/Public/Types/UIFWidgetTree.h b/Engine/Plugins/Experimental/UIFramework/Source/Public/Types/UIFWidgetTree.h
--- a/Engine/Plugins/Experimental/UIFramework/Source/Public/Types/UIFWidgetTree.h
+++ b/Engine/Plugins/Experimental/UIFramework/Source/Public/Types/UIFWidgetTree.h
@@ -10,11 +10,21 @@
 #include "UIFWidgetTree.generated.h"
 
 struct FReplicationFlags;
+struct FUIFrameworkWidgetTreeEntry;
+class AActor;
 class FOutBunch;
 class UActorChannel;
-class UUIFrameworkPlayerComponent;
 class UUIFrameworkWidget;
 
+class IUIFrameworkWidgetTreeOwner
+{
+public:
+	/** A widget was added to the tree. */
+	virtual void LocalWidgetWasAddedToTree(const FUIFrameworkWidgetTreeEntry& Entry) {};
+	/** A widget was removed to the tree. */
+	virtual void LocalWidgetRemovedFromTree(const FUIFrameworkWidgetTreeEntry& Entry) {};
+};
+
 /**
  *
  */
@@ -55,8 +65,9 @@ struct UIFRAMEWORK_API FUIFrameworkWidgetTree : public FFastArraySerializer
 
 public:
 	FUIFrameworkWidgetTree() = default;
-	FUIFrameworkWidgetTree(UUIFrameworkPlayerComponent* InOwnerComponent)
-		: OwnerComponent(InOwnerComponent)
+	FUIFrameworkWidgetTree(AActor* InReplicatedOwner, IUIFrameworkWidgetTreeOwner* InOwner)
+		: ReplicatedOwner(InReplicatedOwner)
+		, Owner(InOwner)
 	{
 	}
 
@@ -106,7 +117,9 @@ private:
 	TMap<FUIFrameworkWidgetId, TWeakObjectPtr<UUIFrameworkWidget>> WidgetByIdMap;
 
 	UPROPERTY(NotReplicated, Transient)
-	TObjectPtr<UUIFrameworkPlayerComponent> OwnerComponent = nullptr;
+	TObjectPtr<AActor> ReplicatedOwner;
+
+	IUIFrameworkWidgetTreeOwner* Owner = nullptr;
 };
 
 template<>
diff --git a/Engine/Plugins/Experimental/UIFramework/Source/Public/UIFPlayerComponent.h b/Engine/Plugins/Experimental/UIFramework/Source/Public/UIFPlayerComponent.h
--- a/Engine/Plugins/Experimental/UIFramework/Source/Public/UIFPlayerComponent.h
+++ b/Engine/Plugins/Experimental/UIFramework/Source/Public/UIFPlayerComponent.h
@@ -93,7 +93,7 @@ struct TStructOpsTypeTraits<FUIFrameworkGameLayerSlotList> : public TStructOpsTy
  * 
  */
 UCLASS(Blueprintable, meta = (BlueprintSpawnableComponent))
-class UIFRAMEWORK_API UUIFrameworkPlayerComponent : public UActorComponent
+class UIFRAMEWORK_API UUIFrameworkPlayerComponent : public UActorComponent, public IUIFrameworkWidgetTreeOwner
 {
 	GENERATED_BODY()
 
@@ -135,8 +135,9 @@ public:
 	//~ End UActorComponent
 
 	void AuthorityRemoveChild(UUIFrameworkWidget* Widget);
-	void LocalWidgetWasAddedToTree(const FUIFrameworkWidgetTreeEntry& Entry);
-	void LocalWidgetRemovedFromTree(const FUIFrameworkWidgetTreeEntry& Entry);
+
+	virtual void LocalWidgetWasAddedToTree(const FUIFrameworkWidgetTreeEntry& Entry) override;
+	virtual void LocalWidgetRemovedFromTree(const FUIFrameworkWidgetTreeEntry& Entry) override;
 
 private:
 	void LocalOnClassLoaded(TSoftClassPtr<UWidget> WidgetClass);
