UIFramework:
Fix bugs while adding, removing, then readding, then removing, then readding, then.... the same widget to the same player.
Fix replication issue with stack_box_slot index
Prepare the field to move widget from one player to another.

diff --git a/Engine/Plugins/Experimental/UIFramework/Source/Private/Types/UIFWidgetTree.cpp b/Engine/Plugins/Experimental/UIFramework/Source/Private/Types/UIFWidgetTree.cpp
--- a/Engine/Plugins/Experimental/UIFramework/Source/Private/Types/UIFWidgetTree.cpp
+++ b/Engine/Plugins/Experimental/UIFramework/Source/Private/Types/UIFWidgetTree.cpp
@@ -12,6 +12,25 @@
 #include "GameFramework/PlayerController.h"
 #include "Net/UnrealNetwork.h"
 
+#if UE_UIFRAMEWORK_WITH_DEBUG
+namespace UE::UIFramework::Private
+{
+	TArray<FUIFrameworkWidgetTree*> GTrees;
+	void TestWidgetTree()
+	{
+		for (FUIFrameworkWidgetTree* Tree : GTrees)
+		{
+			Tree->AuthorityTest();
+		}
+	}
+
+	static FAutoConsoleCommand CCmdTestWidgetTree(
+		TEXT("UIFramework.TestWidgetTree"),
+		TEXT("Test if all containers are properly setup"),
+		FConsoleCommandDelegate::CreateStatic(TestWidgetTree),
+		ECVF_Cheat);
+}
+#endif
 
 /**
  *
@@ -38,6 +57,22 @@ bool FUIFrameworkWidgetTreeEntry::IsChildValid() const
 /**
  *
  */
+FUIFrameworkWidgetTree::FUIFrameworkWidgetTree(AActor* InReplicatedOwner, IUIFrameworkWidgetTreeOwner* InOwner)
+	: ReplicatedOwner(InReplicatedOwner)
+	, Owner(InOwner)
+{
+#if UE_UIFRAMEWORK_WITH_DEBUG
+	UE::UIFramework::Private::GTrees.Add(this);
+#endif
+}
+
+FUIFrameworkWidgetTree::~FUIFrameworkWidgetTree()
+{
+#if UE_UIFRAMEWORK_WITH_DEBUG
+	UE::UIFramework::Private::GTrees.RemoveSingleSwap(this);
+#endif
+}
+
 void FUIFrameworkWidgetTree::PreReplicatedRemove(const TArrayView<int32> RemovedIndices, int32 FinalSize)
 {
 	for (int32 Index : RemovedIndices)
@@ -50,26 +85,8 @@ void FUIFrameworkWidgetTree::PreReplicatedRemove(const TArrayView<int32> Removed
 				Owner->LocalWidgetRemovedFromTree(Entry);
 			}
 
-			// Was it remove but added by something else
-			{
-				int32 Count = 0;
-				for (const FUIFrameworkWidgetTreeEntry& Other : Entries)
-				{
-					if (Other.Child == Entry.Child)
-					{
-						++Count;
-						if (Count > 2)
-						{
-							break;
-						}
-					}
-				}
-				if (Count <= 1)
-				{
-					Entry.Child->LocalDestroyUMGWidget();
-					WidgetByIdMap.Remove(Entry.Child->GetWidgetId());
-				}
-			}
+			Entry.Child->LocalDestroyUMGWidget();
+			WidgetByIdMap.Remove(Entry.Child->GetWidgetId());
 		}
 	}
 }
@@ -138,117 +155,114 @@ bool FUIFrameworkWidgetTree::ReplicateSubWidgets(UActorChannel* Channel, FOutBun
 	return bWroteSomething;
 }
 
-void FUIFrameworkWidgetTree::AddRoot(UUIFrameworkWidget* Widget)
+void FUIFrameworkWidgetTree::AuthorityAddRoot(UUIFrameworkWidget* Widget)
 {
 	check(Widget);
-	AddChildInternal(nullptr, Widget);
+	AuthorityAddChildInternal(nullptr, Widget);
 }
 
-void FUIFrameworkWidgetTree::AddWidget(UUIFrameworkWidget* Parent, UUIFrameworkWidget* Child)
+void FUIFrameworkWidgetTree::AuthorityAddWidget(UUIFrameworkWidget* Parent, UUIFrameworkWidget* Child)
 {
 	check(Parent);
 	check(Child);
-	AddChildInternal(Parent, Child);
+	AuthorityAddChildInternal(Parent, Child);
 }
 
-void FUIFrameworkWidgetTree::AddChildInternal(UUIFrameworkWidget* Parent, UUIFrameworkWidget* Child)
+void FUIFrameworkWidgetTree::AuthorityAddChildInternal(UUIFrameworkWidget* Parent, UUIFrameworkWidget* Child)
 {
-	int32 PreviousEntryIndex = Entries.IndexOfByPredicate([Child](const FUIFrameworkWidgetTreeEntry& Other){ return Other.Child == Child; });
-	if (PreviousEntryIndex != INDEX_NONE)
+	if (int32* PreviousEntryIndexPtr = AuthorityIndexByWidgetMap.Find(Child))
 	{
-		const FUIFrameworkWidgetTreeEntry& PreviousEntry = Entries[PreviousEntryIndex];
+		check(Entries.IsValidIndex(*PreviousEntryIndexPtr));
+		FUIFrameworkWidgetTreeEntry& PreviousEntry = Entries[*PreviousEntryIndexPtr];
 		if (PreviousEntry.Parent != Parent)
 		{
-			// New Parent for an existing child.
-			Entries.RemoveAtSwap(PreviousEntryIndex);
-			FUIFrameworkWidgetTreeEntry& NewEntry = Entries.Emplace_GetRef(Parent, Child);
-			MarkItemDirty(NewEntry);
-		}
-		else
-		{
-			UE_LOG(LogUIFramework, Warning, TEXT("A widget was added to the WidgetTree but it's already added."));
+			// Same child, different parent. Need to build a new entry for replication.
+			PreviousEntry = FUIFrameworkWidgetTreeEntry(Parent, Child);
+			MarkItemDirty(PreviousEntry);
 		}
 	}
 	else
 	{
-		FUIFrameworkWidgetTreeEntry& NewEntry = Entries.Emplace_GetRef(Parent, Child);
+		int32 NewEntryIndex = Entries.Emplace(Parent, Child);
+
+		FUIFrameworkWidgetTreeEntry& NewEntry = Entries[NewEntryIndex];
 		MarkItemDirty(NewEntry);
-		WidgetByIdMap.FindOrAdd(Child->GetWidgetId()) = Child;
+
+		AuthorityIndexByWidgetMap.Add(Child) = NewEntryIndex;
+		WidgetByIdMap.Add(Child->GetWidgetId()) = Child;
 
 		if (ensure(ReplicatedOwner) && ReplicatedOwner->IsUsingRegisteredSubObjectList())
 		{
 			ReplicatedOwner->AddReplicatedSubObject(Child);
 		}
-		AddChildRecursiveInternal(Child);
+
+		AuthorityAddChildRecursiveInternal(Child);
 	}
 }
 
-void FUIFrameworkWidgetTree::AddChildRecursiveInternal(UUIFrameworkWidget* Widget)
+void FUIFrameworkWidgetTree::AuthorityAddChildRecursiveInternal(UUIFrameworkWidget* InParentWidget)
 {
 	FUIFrameworkWidgetTree* Self = this;
-	Widget->AuthorityForEachChildren([Self, Widget](UUIFrameworkWidget* ChildWidget)
+	InParentWidget->AuthorityForEachChildren([Self, InParentWidget](UUIFrameworkWidget* ChildWidget)
 		{
 			if (ChildWidget != nullptr)
 			{
-				FUIFrameworkWidgetTreeEntry& NewEntry = Self->Entries.Emplace_GetRef(Widget, ChildWidget);
-				Self->MarkItemDirty(NewEntry);
-				Self->WidgetByIdMap.FindOrAdd(ChildWidget->GetWidgetId()) = ChildWidget;
-				if (ensure(Self->ReplicatedOwner) && Self->ReplicatedOwner->IsUsingRegisteredSubObjectList())
-				{
-					Self->ReplicatedOwner->AddReplicatedSubObject(ChildWidget);
-				}
-				Self->AddChildRecursiveInternal(ChildWidget);
+				Self->AuthorityAddChildInternal(InParentWidget, ChildWidget);
 			}
 		});
 }
 
-void FUIFrameworkWidgetTree::RemoveWidget(UUIFrameworkWidget* Widget)
+void FUIFrameworkWidgetTree::AuthorityRemoveWidget(UUIFrameworkWidget* Widget)
 {
 	check(Widget);
-	bool bDirtyArray = false;
+
+	if (AuthorityRemoveChildRecursiveInternal(Widget))
+	{
+		MarkArrayDirty();
+	}
+}
+
+bool FUIFrameworkWidgetTree::AuthorityRemoveChildRecursiveInternal(UUIFrameworkWidget* Widget)
+{
+	if (int32* PreviousEntryIndexPtr = AuthorityIndexByWidgetMap.Find(Widget))
 	{
-		FMemMark Mark(FMemStack::Get());
-		TArray<UUIFrameworkWidget*, TInlineAllocator<60, TMemStackAllocator<>>> WidgetsToRemove;
-		WidgetsToRemove.Add(Widget);
+		check(Entries.IsValidIndex(*PreviousEntryIndexPtr));
+
+		AuthorityIndexByWidgetMap.Remove(Widget);
+		WidgetByIdMap.Remove(Widget->GetWidgetId());
+
+		if (ensure(ReplicatedOwner) && ReplicatedOwner->IsUsingRegisteredSubObjectList())
+		{
+			ReplicatedOwner->RemoveReplicatedSubObject(Widget);
+		}
 
-		while (WidgetsToRemove.Num() > 0)
+		Entries.RemoveAtSwap(*PreviousEntryIndexPtr);
+
+		// Fix up the swap item
+		if (Entries.IsValidIndex(*PreviousEntryIndexPtr))
 		{
-			UUIFrameworkWidget* CurrentWidgetToRemove = WidgetsToRemove.Last();
-			WidgetByIdMap.Remove(CurrentWidgetToRemove->GetWidgetId());
-			if (ensure(ReplicatedOwner) && ReplicatedOwner->IsUsingRegisteredSubObjectList())
+			if (int32* FixUpChildIndexPtr = AuthorityIndexByWidgetMap.Find(Entries[*PreviousEntryIndexPtr].Child))
 			{
-				ReplicatedOwner->RemoveReplicatedSubObject(CurrentWidgetToRemove);
+				*FixUpChildIndexPtr = *PreviousEntryIndexPtr;
 			}
+		}
 
-			WidgetsToRemove.RemoveAt(WidgetsToRemove.Num()-1);
-			for (int32 Index = Entries.Num() - 1; Index >= 0; --Index)
+		FUIFrameworkWidgetTree* Self = this;
+		Widget->AuthorityForEachChildren([Self](UUIFrameworkWidget* ChildWidget)
 			{
-				if (Entries[Index].Parent == CurrentWidgetToRemove)
-				{
-					bDirtyArray = true;
-					if (Entries[Index].Child)
-					{
-						WidgetsToRemove.Add(Entries[Index].Child);
-					}
-
-					Entries.RemoveAtSwap(Index);
-				}
-				else if (Entries[Index].Child == CurrentWidgetToRemove)
+				if (ChildWidget)
 				{
-					bDirtyArray = true;
-					Entries.RemoveAtSwap(Index);
+					Self->AuthorityRemoveChildRecursiveInternal(ChildWidget);
 				}
-			}
-		}
-	}
+			});	
 
-	if (bDirtyArray)
-	{
-		MarkArrayDirty();
+		return true;
 	}
+	return false;
+
 }
 
-FUIFrameworkWidgetTreeEntry* FUIFrameworkWidgetTree::GetEntryByReplicationId(int32 ReplicationId)
+FUIFrameworkWidgetTreeEntry* FUIFrameworkWidgetTree::LocalGetEntryByReplicationId(int32 ReplicationId)
 {
 	if (const int32* Index = ItemMap.Find(ReplicationId))
 	{
@@ -257,9 +271,9 @@ FUIFrameworkWidgetTreeEntry* FUIFrameworkWidgetTree::GetEntryByReplicationId(int
 	return nullptr;
 }
 
-const FUIFrameworkWidgetTreeEntry* FUIFrameworkWidgetTree::GetEntryByReplicationId(int32 ReplicationId) const
+const FUIFrameworkWidgetTreeEntry* FUIFrameworkWidgetTree::LocalGetEntryByReplicationId(int32 ReplicationId) const
 {
-	return const_cast<FUIFrameworkWidgetTree*>(this)->GetEntryByReplicationId(ReplicationId);
+	return const_cast<FUIFrameworkWidgetTree*>(this)->LocalGetEntryByReplicationId(ReplicationId);
 }
 
 UUIFrameworkWidget* FUIFrameworkWidgetTree::FindWidgetById(FUIFrameworkWidgetId WidgetId)
@@ -271,4 +285,86 @@ UUIFrameworkWidget* FUIFrameworkWidgetTree::FindWidgetById(FUIFrameworkWidgetId
 const UUIFrameworkWidget* FUIFrameworkWidgetTree::FindWidgetById(FUIFrameworkWidgetId WidgetId) const
 {
 	return const_cast<FUIFrameworkWidgetTree*>(this)->FindWidgetById(WidgetId);
-}
\ No newline at end of file
+}
+
+#if UE_UIFRAMEWORK_WITH_DEBUG
+void FUIFrameworkWidgetTree::AuthorityTest() const
+{
+	if (ReplicatedOwner == nullptr || !ReplicatedOwner->HasAuthority())
+	{
+		return;
+	}
+
+	TSet<FUIFrameworkWidgetId> UniqueIds;
+	for (int32 Index = 0; Index < Entries.Num(); ++Index)
+	{
+		const FUIFrameworkWidgetTreeEntry& Entry = Entries[Index];
+
+		ensureAlwaysMsgf(Entry.ParentId.IsValid(), TEXT("Invalid ParentId"));
+		if (!Entry.ParentId.IsRoot())
+		{
+			ensureAlwaysMsgf(Entry.Parent, TEXT("Invalid Parent"));
+		}
+		else
+		{
+			ensureAlwaysMsgf(!Entry.Parent, TEXT("Valid Parent"));
+		}
+		ensureAlwaysMsgf(Entry.Child, TEXT("Invalid Child"));
+		ensureAlwaysMsgf(Entry.ChildId.IsValid(), TEXT("Invalid ChildId"));
+
+		if (Entry.ChildId.IsValid())
+		{
+			ensureAlwaysMsgf(!UniqueIds.Contains(Entry.ChildId), TEXT("Duplicated id"));
+			UniqueIds.Add(Entry.ChildId);
+		}
+
+		if (Entry.Child)
+		{
+			ensureAlwaysMsgf(Entry.ChildId == Entry.Child->GetWidgetId(), TEXT("Id do not matches"));
+		}
+		
+		if (Entry.Parent)
+		{
+			bool bFound = false;
+			const UUIFrameworkWidget* ToFindWidget = Entry.Child;
+			Entry.Parent->AuthorityForEachChildren([&bFound, ToFindWidget](UUIFrameworkWidget* ChildWidget)
+				{
+					bFound = bFound || ToFindWidget == ChildWidget;
+				});
+
+			ensureAlwaysMsgf(bFound, TEXT("Widget is in the tree but not in the AuthorityForEachChildren"));
+		}
+
+
+		const int32* FoundIndexPtr = AuthorityIndexByWidgetMap.Find(Entry.Child);
+		if (FoundIndexPtr)
+		{
+			ensureAlwaysMsgf(*FoundIndexPtr == Index, TEXT("Widget index doesn't match what is in the map"));
+		}
+		else
+		{
+			ensureAlwaysMsgf(false, TEXT("Widget no in the map"));
+		}
+		
+
+		if (Entry.Child)
+		{
+			ensureAlwaysMsgf(Entry.Child->GetWidgetId().IsValid(), TEXT("The id is not valid."));
+			const TWeakObjectPtr<UUIFrameworkWidget>* FoundWidgetPtr = WidgetByIdMap.Find(Entry.ChildId);
+			if (FoundWidgetPtr)
+			{
+				ensureAlwaysMsgf(FoundWidgetPtr->Get(), TEXT("The found widget is invalid"));
+				if (UUIFrameworkWidget* Widget = FoundWidgetPtr->Get())
+				{
+					ensureAlwaysMsgf(Widget == Entry.Child, TEXT("Widget in the map doesn't matches with the entry widget."));
+				}
+			}
+			else
+			{
+				ensureAlwaysMsgf(false, TEXT("Widget no in the map"));
+			}
+		}
+		
+	}
+}
+#endif
\ No newline at end of file
diff --git a/Engine/Plugins/Experimental/UIFramework/Source/Private/UIFModule.cpp b/Engine/Plugins/Experimental/UIFramework/Source/Private/UIFModule.cpp
new file mode 100644
--- /dev/null
+++ b/Engine/Plugins/Experimental/UIFramework/Source/Private/UIFModule.cpp
@@ -0,0 +1,97 @@
+// Copyright Epic Games, Inc. All Rights Reserved.
+
+#include "UIFModule.h"
+
+#include "Modules/ModuleManager.h"
+#include "UIFPlayerComponent.h"
+#include "UIFWidget.h"
+
+#if WITH_SERVER_CODE
+UUIFrameworkWidget* FUIFrameworkModule::AuthorityAttachWidget(UUIFrameworkPlayerComponent* ReplicationOwner, FUIFrameworkParentWidget Parent, UUIFrameworkWidget* Child)
+{
+	check(Child);
+	check(Parent.IsParentValid());
+
+	if (Child->AuthorityGetParent().IsParentValid())
+	{
+		AuthorityDetachWidgetFromParent(Child);
+	}
+
+	UObject* ParentOuter = Parent.IsPlayerComponent() ? Parent.AsPlayerComponent()->GetOuter() : Parent.AsWidget()->GetOuter();
+	if (ParentOuter != Child->GetOuter())
+	{
+		if (Child->GetOuter() == GetTransientPackage())
+		{
+			// Recursive Rename to change the outer
+			// To the operation recursively to child
+		}
+	}
+
+	bool bDifferentReplicationOwner = ReplicationOwner != Child->OwnerPlayerComponent;
+
+	Child->AuthorityParent = Parent;
+	Child->OwnerPlayerComponent = ReplicationOwner;
+
+	if (ReplicationOwner)
+	{
+		if (Child->AuthorityParent.IsWidget())
+		{
+			ReplicationOwner->GetWidgetTree().AuthorityAddWidget(Child->AuthorityParent.AsWidget(), Child);
+		}
+		else
+		{
+			check(Child->AuthorityParent.IsPlayerComponent());
+			ReplicationOwner->GetWidgetTree().AuthorityAddRoot(Child);
+		}
+	}
+
+	if (bDifferentReplicationOwner)
+	{
+		AuthoritySetParentReplicationOwnerRecursive(Child);
+	}
+	return Child;
+}
+
+void FUIFrameworkModule::AuthoritySetParentReplicationOwnerRecursive(UUIFrameworkWidget* Widget)
+{
+	Widget->AuthorityForEachChildren([Widget](UUIFrameworkWidget* Child)
+		{
+			if (Child != nullptr)
+			{
+				check(Child->AuthorityGetParent().IsWidget() && Child->AuthorityGetParent().AsWidget() == Widget);
+				Child->OwnerPlayerComponent = Widget->OwnerPlayerComponent;
+				Child->AuthorityParent = FUIFrameworkParentWidget(Widget);
+				AuthoritySetParentReplicationOwnerRecursive(Child);
+			}
+		});
+}
+
+bool FUIFrameworkModule::AuthorityCanWidgetBeAttached(UUIFrameworkPlayerComponent* ReplicationOwner, UUIFrameworkWidget* Parent, UUIFrameworkWidget* Child)
+{
+	return true;
+}
+
+void FUIFrameworkModule::AuthorityDetachWidgetFromParent(UUIFrameworkWidget* Child)
+{
+	check(Child);
+	if (Child->OwnerPlayerComponent)
+	{
+		Child->OwnerPlayerComponent->GetWidgetTree().AuthorityRemoveWidget(Child);
+	}
+
+	if (Child->AuthorityGetParent().IsParentValid())
+	{
+		if (Child->AuthorityGetParent().IsWidget())
+		{
+			Child->AuthorityGetParent().AsWidget()->AuthorityRemoveChild(Child);
+		}
+		else
+		{
+			check(Child->AuthorityGetParent().IsPlayerComponent());
+			Child->AuthorityGetParent().AsPlayerComponent()->AuthorityRemoveChild(Child);
+		}
+	}
+}
+#endif //WITH_SERVER_CODE
+
+IMPLEMENT_MODULE(FUIFrameworkModule, UIFramework)
\ No newline at end of file
diff --git a/Engine/Plugins/Experimental/UIFramework/Source/Private/UIFPlayerComponent.cpp b/Engine/Plugins/Experimental/UIFramework/Source/Private/UIFPlayerComponent.cpp
--- a/Engine/Plugins/Experimental/UIFramework/Source/Private/UIFPlayerComponent.cpp
+++ b/Engine/Plugins/Experimental/UIFramework/Source/Private/UIFPlayerComponent.cpp
@@ -2,6 +2,7 @@
 
 #include "UIFPlayerComponent.h"
 #include "UIFLog.h"
+#include "UIFModule.h"
 #include "UIFWidget.h"
 #include "Blueprint/GameViewportSubsystem.h"
 
@@ -121,8 +122,8 @@ void UUIFrameworkPlayerComponent::AddWidget(FUIFrameworkGameLayerSlot InEntry)
 		}
 		else
 		{
-			InEntry.AuthoritySetWidget(InEntry.AuthorityGetWidget()); // to make sure the id is set
-			InEntry.AuthorityGetWidget()->AuthoritySetParent(this, FUIFrameworkParentWidget(this));
+			// Reset the widget to make sure the id is set and it may have been duplicated during the attach
+			InEntry.AuthoritySetWidget(FUIFrameworkModule::AuthorityAttachWidget(this, this, InEntry.AuthorityGetWidget()));
 			RootList.AddEntry(InEntry);
 		}
 	}
@@ -146,8 +147,8 @@ void UUIFrameworkPlayerComponent::RemoveWidget(UUIFrameworkWidget* Widget)
 		}
 		else
 		{
+			FUIFrameworkModule::AuthorityDetachWidgetFromParent(Widget);
 			RootList.RemoveEntry(Widget);
-			Widget->AuthoritySetParent(this, FUIFrameworkParentWidget());
 		}
 	}
 }
@@ -170,7 +171,7 @@ void UUIFrameworkPlayerComponent::TickComponent(float DeltaTime, enum ELevelTick
 		TGuardValue<bool> TmpGuard = {bAddingWidget, true};
 		for (int32 ReplicationId : AddPending)
 		{
-			if (FUIFrameworkWidgetTreeEntry* Entry = WidgetTree.GetEntryByReplicationId(ReplicationId))
+			if (FUIFrameworkWidgetTreeEntry* Entry = WidgetTree.LocalGetEntryByReplicationId(ReplicationId))
 			{
 				if (ensure(Entry->IsParentValid() && Entry->IsChildValid()))
 				{
@@ -209,7 +210,7 @@ void UUIFrameworkPlayerComponent::TickComponent(float DeltaTime, enum ELevelTick
 
 		for (int32 ReplicationId : TempNetReplicationPending)
 		{
-			if (const FUIFrameworkWidgetTreeEntry* Entry = WidgetTree.GetEntryByReplicationId(ReplicationId))
+			if (const FUIFrameworkWidgetTreeEntry* Entry = WidgetTree.LocalGetEntryByReplicationId(ReplicationId))
 			{
 				LocalWidgetWasAddedToTree(*Entry);
 			}
@@ -299,7 +300,7 @@ void UUIFrameworkPlayerComponent::LocalOnClassLoaded(TSoftClassPtr<UWidget> Widg
 		{
 			for (int32 ReplicationId : FoundWidgetClassToLoad.EntryReplicationIds)
 			{
-				if (const FUIFrameworkWidgetTreeEntry* Entry = WidgetTree.GetEntryByReplicationId(ReplicationId))
+				if (const FUIFrameworkWidgetTreeEntry* Entry = WidgetTree.LocalGetEntryByReplicationId(ReplicationId))
 				{
 					if (Entry->IsParentValid() && Entry->IsChildValid())
 					{
diff --git a/Engine/Plugins/Experimental/UIFramework/Source/Private/UIFWidget.cpp b/Engine/Plugins/Experimental/UIFramework/Source/Private/UIFWidget.cpp
--- a/Engine/Plugins/Experimental/UIFramework/Source/Private/UIFWidget.cpp
+++ b/Engine/Plugins/Experimental/UIFramework/Source/Private/UIFWidget.cpp
@@ -66,74 +66,10 @@ void UUIFrameworkWidget::GetLifetimeReplicatedProps(TArray<FLifetimeProperty>& O
 	DOREPLIFETIME_WITH_PARAMS_FAST(ThisClass, Id, Params);
 }
 
-void UUIFrameworkWidget::AuthoritySetParent(UUIFrameworkPlayerComponent* NewOwner, FUIFrameworkParentWidget NewParent)
-{
-	const bool bDifferentOwner = NewOwner != OwnerPlayerComponent;
-	if (OwnerPlayerComponent)
-	{
-		ensure(OwnerPlayerComponent == NewOwner);
-		NewOwner = OwnerPlayerComponent;
-		NewParent = FUIFrameworkParentWidget();
-	}
-
-	if (AuthorityParent.IsParentValid())
-	{
-		if (AuthorityParent.IsWidget())
-		{
-			AuthorityParent.AsWidget()->AuthorityRemoveChild(this);
-		}
-		else
-		{
-			check(AuthorityParent.IsPlayerComponent());
-			AuthorityParent.AsPlayerComponent()->AuthorityRemoveChild(this);
-		}
-	}
-
-	AuthorityParent = NewParent;
-	OwnerPlayerComponent = NewOwner;
-
-	if (AuthorityParent.IsParentValid() && OwnerPlayerComponent)
-	{
-		if (AuthorityParent.IsWidget())
-		{
-			OwnerPlayerComponent->GetWidgetTree().AddWidget(AuthorityParent.AsWidget(), this);
-		}
-		else
-		{
-			check(AuthorityParent.IsPlayerComponent());
-			OwnerPlayerComponent->GetWidgetTree().AddRoot(this);
-		}
-	}
-	else if (OwnerPlayerComponent)
-	{
-		OwnerPlayerComponent->GetWidgetTree().RemoveWidget(this);
-	}
-
-	if (bDifferentOwner)
-	{
-		SetParentPlayerOwnerRecursive();
-	}
-}
-
-void UUIFrameworkWidget::SetParentPlayerOwnerRecursive()
-{
-	UUIFrameworkWidget* Self = this;
-	AuthorityForEachChildren([Self](UUIFrameworkWidget* Child)
-		{
-			if (Child != nullptr)
-			{
-				check(Child->AuthorityGetParent().IsWidget() && Child->AuthorityGetParent().AsWidget() == Self);
-				Child->OwnerPlayerComponent = Self->OwnerPlayerComponent;
-				Child->AuthorityParent = FUIFrameworkParentWidget(Self);
-				Child->SetParentPlayerOwnerRecursive();
-			}
-		});
-}
-
 void UUIFrameworkWidget::LocalAddChild(FUIFrameworkWidgetId ChildId)
 {
 	// By default we should remove the widget from its previous parent.
-	//Adding a widget to a new slot will automaticly remove it from its previous parent.
+	//Adding a widget to a new slot will automatically remove it from its previous parent.
 	if (OwnerPlayerComponent)
 	{
 		if (UUIFrameworkWidget* Widget = OwnerPlayerComponent->GetWidgetTree().FindWidgetById(ChildId))
diff --git a/Engine/Plugins/Experimental/UIFramework/Source/Private/UIFrameworkModule.cpp b/Engine/Plugins/Experimental/UIFramework/Source/Private/UIFrameworkModule.cpp
deleted file mode 100644
diff --git a/Engine/Plugins/Experimental/UIFramework/Source/Private/Widgets/UIFButton.cpp b/Engine/Plugins/Experimental/UIFramework/Source/Private/Widgets/UIFButton.cpp
--- a/Engine/Plugins/Experimental/UIFramework/Source/Private/Widgets/UIFButton.cpp
+++ b/Engine/Plugins/Experimental/UIFramework/Source/Private/Widgets/UIFButton.cpp
@@ -3,6 +3,7 @@
 #include "Widgets/UIFButton.h"
 #include "Types/UIFWidgetTree.h"
 #include "UIFLog.h"
+#include "UIFModule.h"
 #include "UIFPlayerComponent.h"
 
 #include "Components/Button.h"
@@ -41,7 +42,7 @@ void UUIFrameworkButton::SetContent(FUIFrameworkSimpleSlot InEntry)
 		// Remove previous widget
 		if (Slot.AuthorityGetWidget())
 		{
-			Slot.AuthorityGetWidget()->AuthoritySetParent(GetPlayerComponent(), FUIFrameworkParentWidget());
+			FUIFrameworkModule::AuthorityDetachWidgetFromParent(Slot.AuthorityGetWidget());
 		}
 
 		if (InEntry.AuthorityGetWidget())
@@ -56,11 +57,11 @@ void UUIFrameworkButton::SetContent(FUIFrameworkSimpleSlot InEntry)
 	}
 
 	Slot = InEntry;
-	Slot.AuthoritySetWidget(InEntry.AuthorityGetWidget()); // to make sure the id is set
 
 	if (bWidgetIsDifferent && Slot.AuthorityGetWidget())
 	{
-		Slot.AuthorityGetWidget()->AuthoritySetParent(GetPlayerComponent(), FUIFrameworkParentWidget(this));
+		// Reset the widget to make sure the id is set and it may have been duplicated during the attach
+		Slot.AuthoritySetWidget(FUIFrameworkModule::AuthorityAttachWidget(GetPlayerComponent(), this, Slot.AuthorityGetWidget()));
 	}
 
 	MARK_PROPERTY_DIRTY_FROM_NAME(ThisClass, Slot, this);
diff --git a/Engine/Plugins/Experimental/UIFramework/Source/Private/Widgets/UIFCanvasBox.cpp b/Engine/Plugins/Experimental/UIFramework/Source/Private/Widgets/UIFCanvasBox.cpp
--- a/Engine/Plugins/Experimental/UIFramework/Source/Private/Widgets/UIFCanvasBox.cpp
+++ b/Engine/Plugins/Experimental/UIFramework/Source/Private/Widgets/UIFCanvasBox.cpp
@@ -3,6 +3,7 @@
 #include "Widgets/UIFCanvasBox.h"
 #include "Types/UIFWidgetTree.h"
 #include "UIFLog.h"
+#include "UIFModule.h"
 #include "UIFPlayerComponent.h"
 
 #include "Components/CanvasPanel.h"
@@ -105,8 +106,8 @@ void UUIFrameworkCanvasBox::AddWidget(FUIFrameworkCanvasBoxSlot InEntry)
 	}
 	else
 	{
-		InEntry.AuthoritySetWidget(InEntry.AuthorityGetWidget()); // to make sure the id is set
-		InEntry.AuthorityGetWidget()->AuthoritySetParent(GetPlayerComponent(), FUIFrameworkParentWidget(this));
+		// Reset the widget to make sure the id is set and it may have been duplicated during the attach
+		InEntry.AuthoritySetWidget(FUIFrameworkModule::AuthorityAttachWidget(GetPlayerComponent(), this, InEntry.AuthorityGetWidget()));
 		AddEntry(InEntry);
 	}
 }
@@ -126,8 +127,8 @@ void UUIFrameworkCanvasBox::RemoveWidget(UUIFrameworkWidget* Widget)
 		}
 		else
 		{
+			FUIFrameworkModule::AuthorityDetachWidgetFromParent(Widget);
 			RemoveEntry(Widget);
-			Widget->AuthoritySetParent(GetPlayerComponent(), FUIFrameworkParentWidget());
 		}
 	}
 }
diff --git a/Engine/Plugins/Experimental/UIFramework/Source/Private/Widgets/UIFStackBox.cpp b/Engine/Plugins/Experimental/UIFramework/Source/Private/Widgets/UIFStackBox.cpp
--- a/Engine/Plugins/Experimental/UIFramework/Source/Private/Widgets/UIFStackBox.cpp
+++ b/Engine/Plugins/Experimental/UIFramework/Source/Private/Widgets/UIFStackBox.cpp
@@ -3,6 +3,7 @@
 #include "Widgets/UIFStackBox.h"
 #include "Types/UIFWidgetTree.h"
 #include "UIFLog.h"
+#include "UIFModule.h"
 #include "UIFPlayerComponent.h"
 
 #include "Components/StackBox.h"
@@ -42,17 +43,23 @@ bool FUIFrameworkStackBoxSlotList::NetDeltaSerialize(FNetDeltaSerializeInfo& Del
 
 void FUIFrameworkStackBoxSlotList::AddEntry(FUIFrameworkStackBoxSlot Entry)
 {
-	FUIFrameworkStackBoxSlot& NewEntry = Slots.Add_GetRef(MoveTemp(Entry));
+	int32 NewEntryIndex = Slots.Add(MoveTemp(Entry));
+	FUIFrameworkStackBoxSlot& NewEntry = Slots[NewEntryIndex];
+	NewEntry.Index = NewEntryIndex;
 	MarkItemDirty(NewEntry);
 }
 
 bool FUIFrameworkStackBoxSlotList::RemoveEntry(UUIFrameworkWidget* Widget)
 {
 	check(Widget);
-	const int32 Index = Slots.IndexOfByPredicate([Widget](const FUIFrameworkStackBoxSlot& Entry) { return Entry.AuthorityGetWidget() == Widget; });
+	int32 Index = Slots.IndexOfByPredicate([Widget](const FUIFrameworkStackBoxSlot& Entry) { return Entry.AuthorityGetWidget() == Widget; });
 	if (Index != INDEX_NONE)
 	{
 		Slots.RemoveAt(Index);
+		for (; Index < Slots.Num(); ++Index)
+		{
+			Slots[Index].Index = Index;
+		}
 		MarkArrayDirty();
 	}
 	return Index != INDEX_NONE;
@@ -103,8 +110,8 @@ void UUIFrameworkStackBox::AddWidget(FUIFrameworkStackBoxSlot InEntry)
 	}
 	else
 	{
-		InEntry.AuthoritySetWidget(InEntry.AuthorityGetWidget()); // to make sure the id is set
-		InEntry.AuthorityGetWidget()->AuthoritySetParent(GetPlayerComponent(), FUIFrameworkParentWidget(this));
+		// Reset the widget to make sure the id is set and it may have been duplicated during the attach
+		InEntry.AuthoritySetWidget(FUIFrameworkModule::AuthorityAttachWidget(GetPlayerComponent(), this, InEntry.AuthorityGetWidget()));
 		AddEntry(InEntry);
 	}
 }
@@ -124,8 +131,8 @@ void UUIFrameworkStackBox::RemoveWidget(UUIFrameworkWidget* Widget)
 		}
 		else
 		{
+			FUIFrameworkModule::AuthorityDetachWidgetFromParent(Widget);
 			RemoveEntry(Widget);
-			Widget->AuthoritySetParent(GetPlayerComponent(), FUIFrameworkParentWidget());
 		}
 	}
 }
diff --git a/Engine/Plugins/Experimental/UIFramework/Source/Public/Types/UIFWidgetTree.h b/Engine/Plugins/Experimental/UIFramework/Source/Public/Types/UIFWidgetTree.h
--- a/Engine/Plugins/Experimental/UIFramework/Source/Public/Types/UIFWidgetTree.h
+++ b/Engine/Plugins/Experimental/UIFramework/Source/Public/Types/UIFWidgetTree.h
@@ -9,6 +9,10 @@
 
 #include "UIFWidgetTree.generated.h"
 
+#ifndef UE_UIFRAMEWORK_WITH_DEBUG
+	#define UE_UIFRAMEWORK_WITH_DEBUG !(UE_BUILD_SHIPPING || UE_BUILD_TEST)
+#endif
+
 struct FReplicationFlags;
 struct FUIFrameworkWidgetTreeEntry;
 class AActor;
@@ -65,11 +69,8 @@ struct UIFRAMEWORK_API FUIFrameworkWidgetTree : public FFastArraySerializer
 
 public:
 	FUIFrameworkWidgetTree() = default;
-	FUIFrameworkWidgetTree(AActor* InReplicatedOwner, IUIFrameworkWidgetTreeOwner* InOwner)
-		: ReplicatedOwner(InReplicatedOwner)
-		, Owner(InOwner)
-	{
-	}
+	FUIFrameworkWidgetTree(AActor* InReplicatedOwner, IUIFrameworkWidgetTreeOwner* InOwner);
+	~FUIFrameworkWidgetTree();
 
 public:
 	//~ Begin of FFastArraySerializer
@@ -86,39 +87,44 @@ public:
 	bool ReplicateSubWidgets(UActorChannel* Channel, FOutBunch* Bunch, FReplicationFlags* RepFlags);
 
 	/** Add a new widget to the top hierarchy. */
-	void AddRoot(UUIFrameworkWidget* Widget);
+	void AuthorityAddRoot(UUIFrameworkWidget* Widget);
 	/**
 	 * Change the parent / child relationship of the child widget.
 	 * If the child widget had a parent, that relationship entry will replaced it by a new one.
 	 */
-	void AddWidget(UUIFrameworkWidget* Parent, UUIFrameworkWidget* Child);
+	void AuthorityAddWidget(UUIFrameworkWidget* Parent, UUIFrameworkWidget* Child);
 	/**
 	 * Remove the widget and all of its children and grand-children from the tree.
 	 * It will clean all the parent relationship from the tree.
 	 */
-	void RemoveWidget(UUIFrameworkWidget* Widget);
+	void AuthorityRemoveWidget(UUIFrameworkWidget* Widget);
 
-	FUIFrameworkWidgetTreeEntry* GetEntryByReplicationId(int32 ReplicationId);
-	const FUIFrameworkWidgetTreeEntry* GetEntryByReplicationId(int32 ReplicationId) const;
+	//~ It is not safe to use the ReplicationId on the Authority because any add or remove would clear the ItemMap.
+	FUIFrameworkWidgetTreeEntry* LocalGetEntryByReplicationId(int32 WidgetId);
+	const FUIFrameworkWidgetTreeEntry* LocalGetEntryByReplicationId(int32 WidgetId) const;
 
 	/** Find the widget by its unique Id. The widget needs to be in the Tree. */
 	UUIFrameworkWidget* FindWidgetById(FUIFrameworkWidgetId WidgetId);
 	const UUIFrameworkWidget* FindWidgetById(FUIFrameworkWidgetId WidgetId) const;
 
+#if UE_UIFRAMEWORK_WITH_DEBUG
+	void AuthorityTest() const;
+#endif
+
 private:
-	void AddChildInternal(UUIFrameworkWidget* Parent, UUIFrameworkWidget* Child);
-	void AddChildRecursiveInternal(UUIFrameworkWidget* Widget);
+	void AuthorityAddChildInternal(UUIFrameworkWidget* Parent, UUIFrameworkWidget* Child);
+	void AuthorityAddChildRecursiveInternal(UUIFrameworkWidget* Widget);
+	bool AuthorityRemoveChildRecursiveInternal(UUIFrameworkWidget* Widget);
 
 private:
 	UPROPERTY()
 	TArray<FUIFrameworkWidgetTreeEntry> Entries;
 
-	UPROPERTY(NotReplicated)
-	TMap<FUIFrameworkWidgetId, TWeakObjectPtr<UUIFrameworkWidget>> WidgetByIdMap;
-
 	UPROPERTY(NotReplicated, Transient)
 	TObjectPtr<AActor> ReplicatedOwner;
 
+	TMap<TWeakObjectPtr<UUIFrameworkWidget>, int32> AuthorityIndexByWidgetMap;
+	TMap<FUIFrameworkWidgetId, TWeakObjectPtr<UUIFrameworkWidget>> WidgetByIdMap;
 	IUIFrameworkWidgetTreeOwner* Owner = nullptr;
 };
 
diff --git a/Engine/Plugins/Experimental/UIFramework/Source/Public/UIFModule.h b/Engine/Plugins/Experimental/UIFramework/Source/Public/UIFModule.h
new file mode 100644
--- /dev/null
+++ b/Engine/Plugins/Experimental/UIFramework/Source/Public/UIFModule.h
@@ -0,0 +1,37 @@
+// Copyright Epic Games, Inc. All Rights Reserved.
+
+#pragma once
+
+#include "CoreMinimal.h"
+#include "Modules/ModuleInterface.h"
+#include "Types/UIFParentWidget.h"
+
+class UUIFrameworkWidget;
+class UUIFrameworkPlayerComponent;
+
+/**
+ * 
+ */
+class UIFRAMEWORK_API FUIFrameworkModule : public IModuleInterface
+{
+public:
+	/**
+	 * Set the new widget parent.
+	 * It will attach the widget to the correct ReplicationOwner and add it to the WidgetTree.
+	 * If the ReplicationOwner cannot be changed, a duplicate of the widget will be created and the AuthorityOnReattachWidgets will be broadcast.
+	 * This function is recursive if the owner changed.
+	 */
+	static UUIFrameworkWidget* AuthorityAttachWidget(UUIFrameworkPlayerComponent* ReplicationOwner, FUIFrameworkParentWidget Parent, UUIFrameworkWidget* Child);
+	static bool AuthorityCanWidgetBeAttached(UUIFrameworkPlayerComponent* ReplicationOwner, UUIFrameworkWidget* Parent, UUIFrameworkWidget* Child);
+	/**
+	 * Will remove the widget from the tree and the replication owner.
+	 */
+	static void AuthorityDetachWidgetFromParent(UUIFrameworkWidget* Child);
+	
+	using FOldWidgetToNewWidgetMap = TMap<UUIFrameworkWidget*, UUIFrameworkWidget*>;
+	DECLARE_DELEGATE_OneParam(FOnReattachWidgets, FOldWidgetToNewWidgetMap);
+	FOnReattachWidgets AuthorityOnReattachWidgets;
+
+private:
+	static void AuthoritySetParentReplicationOwnerRecursive(UUIFrameworkWidget* Widget);
+};
\ No newline at end of file
diff --git a/Engine/Plugins/Experimental/UIFramework/Source/Public/UIFWidget.h b/Engine/Plugins/Experimental/UIFramework/Source/Public/UIFWidget.h
--- a/Engine/Plugins/Experimental/UIFramework/Source/Public/UIFWidget.h
+++ b/Engine/Plugins/Experimental/UIFramework/Source/Public/UIFWidget.h
@@ -9,6 +9,7 @@
 
 #include "UIFWidget.generated.h"
 
+class FUIFrameworkModule;
 
 /**
  * 
@@ -17,6 +18,7 @@ UCLASS(Abstract, BlueprintType)
 class UIFRAMEWORK_API UUIFrameworkWidget : public UObject
 {
 	GENERATED_BODY()
+	friend FUIFrameworkModule;
 
 public:
 	//~ Begin UObject
@@ -44,8 +46,6 @@ public:
 	}
 
 	//~ Authority functions
-	void AuthoritySetParent(UUIFrameworkPlayerComponent* Owner, FUIFrameworkParentWidget NewParent);
-
 	FUIFrameworkParentWidget AuthorityGetParent() const
 	{
 		return AuthorityParent;
@@ -67,9 +67,6 @@ protected:
 	virtual void AuthorityRemoveChild(UUIFrameworkWidget* Widget) {}
 	virtual void LocalOnUMGWidgetCreated() { }
 
-private:
-	void SetParentPlayerOwnerRecursive();
-
 protected:
 	UPROPERTY(BlueprintReadOnly, EditDefaultsOnly, Category = "UI Framework")
 	TSoftClassPtr<UWidget> WidgetClass; // todo: make this private and use a constructor argument
diff --git a/Engine/Plugins/Experimental/UIFramework/Source/Public/Widgets/UIFStackBox.h b/Engine/Plugins/Experimental/UIFramework/Source/Public/Widgets/UIFStackBox.h
--- a/Engine/Plugins/Experimental/UIFramework/Source/Public/Widgets/UIFStackBox.h
+++ b/Engine/Plugins/Experimental/UIFramework/Source/Public/Widgets/UIFStackBox.h
@@ -21,6 +21,8 @@ struct FUIFrameworkStackBoxSlot : public FUIFrameworkSlotBase
 {
 	GENERATED_BODY()
 
+	friend FUIFrameworkStackBoxSlotList;
+
 	/** Horizontal alignment of the widget inside the slot. */
 	UPROPERTY(EditAnywhere, BlueprintReadWrite, Category = "UI Framework")
 	TEnumAsByte<EHorizontalAlignment> HorizontalAlignment = EHorizontalAlignment::HAlign_Fill;
@@ -36,6 +38,11 @@ struct FUIFrameworkStackBoxSlot : public FUIFrameworkSlotBase
 	/** How much space this slot should occupy in the direction of the panel. */
 	UPROPERTY(EditAnywhere, BlueprintReadWrite, Category = "UI Framework")
 	FSlateChildSize Size = FSlateChildSize(ESlateSizeRule::Automatic);
+
+private:
+	/** Index in the array the Slot is. The position in the array can change when replicated. */
+	UPROPERTY()
+	int32 Index;
 };
 
 
