Backout CL 22835888

diff --git a/Engine/Plugins/Experimental/UIFramework/Source/Private/Types/UIFWidgetTree.cpp b/Engine/Plugins/Experimental/UIFramework/Source/Private/Types/UIFWidgetTree.cpp
--- a/Engine/Plugins/Experimental/UIFramework/Source/Private/Types/UIFWidgetTree.cpp
+++ b/Engine/Plugins/Experimental/UIFramework/Source/Private/Types/UIFWidgetTree.cpp
@@ -12,25 +12,6 @@
 #include "GameFramework/PlayerController.h"
 #include "Net/UnrealNetwork.h"
 
-#if UE_UIFRAMEWORK_WITH_DEBUG
-namespace UE::UIFramework::Private
-{
-	TArray<FUIFrameworkWidgetTree*> GTrees;
-	void TestWidgetTree()
-	{
-		for (FUIFrameworkWidgetTree* Tree : GTrees)
-		{
-			Tree->AuthorityTest();
-		}
-	}
-
-	static FAutoConsoleCommand CCmdTestWidgetTree(
-		TEXT("UIFramework.TestWidgetTree"),
-		TEXT("Test if all containers are properly setup"),
-		FConsoleCommandDelegate::CreateStatic(TestWidgetTree),
-		ECVF_Cheat);
-}
-#endif
 
 /**
  *
@@ -57,22 +38,6 @@ bool FUIFrameworkWidgetTreeEntry::IsChildValid() const
 /**
  *
  */
-FUIFrameworkWidgetTree::FUIFrameworkWidgetTree(AActor* InReplicatedOwner, IUIFrameworkWidgetTreeOwner* InOwner)
-	: ReplicatedOwner(InReplicatedOwner)
-	, Owner(InOwner)
-{
-#if UE_UIFRAMEWORK_WITH_DEBUG
-	UE::UIFramework::Private::GTrees.Add(this);
-#endif
-}
-
-FUIFrameworkWidgetTree::~FUIFrameworkWidgetTree()
-{
-#if UE_UIFRAMEWORK_WITH_DEBUG
-	UE::UIFramework::Private::GTrees.RemoveSingleSwap(this);
-#endif
-}
-
 void FUIFrameworkWidgetTree::PreReplicatedRemove(const TArrayView<int32> RemovedIndices, int32 FinalSize)
 {
 	for (int32 Index : RemovedIndices)
@@ -85,8 +50,26 @@ void FUIFrameworkWidgetTree::PreReplicatedRemove(const TArrayView<int32> Removed
 				Owner->LocalWidgetRemovedFromTree(Entry);
 			}
 
-			Entry.Child->LocalDestroyUMGWidget();
-			WidgetByIdMap.Remove(Entry.Child->GetWidgetId());
+			// Was it remove but added by something else
+			{
+				int32 Count = 0;
+				for (const FUIFrameworkWidgetTreeEntry& Other : Entries)
+				{
+					if (Other.Child == Entry.Child)
+					{
+						++Count;
+						if (Count > 2)
+						{
+							break;
+						}
+					}
+				}
+				if (Count <= 1)
+				{
+					Entry.Child->LocalDestroyUMGWidget();
+					WidgetByIdMap.Remove(Entry.Child->GetWidgetId());
+				}
+			}
 		}
 	}
 }
@@ -155,114 +138,117 @@ bool FUIFrameworkWidgetTree::ReplicateSubWidgets(UActorChannel* Channel, FOutBun
 	return bWroteSomething;
 }
 
-void FUIFrameworkWidgetTree::AuthorityAddRoot(UUIFrameworkWidget* Widget)
+void FUIFrameworkWidgetTree::AddRoot(UUIFrameworkWidget* Widget)
 {
 	check(Widget);
-	AuthorityAddChildInternal(nullptr, Widget);
+	AddChildInternal(nullptr, Widget);
 }
 
-void FUIFrameworkWidgetTree::AuthorityAddWidget(UUIFrameworkWidget* Parent, UUIFrameworkWidget* Child)
+void FUIFrameworkWidgetTree::AddWidget(UUIFrameworkWidget* Parent, UUIFrameworkWidget* Child)
 {
 	check(Parent);
 	check(Child);
-	AuthorityAddChildInternal(Parent, Child);
+	AddChildInternal(Parent, Child);
 }
 
-void FUIFrameworkWidgetTree::AuthorityAddChildInternal(UUIFrameworkWidget* Parent, UUIFrameworkWidget* Child)
+void FUIFrameworkWidgetTree::AddChildInternal(UUIFrameworkWidget* Parent, UUIFrameworkWidget* Child)
 {
-	if (int32* PreviousEntryIndexPtr = AuthorityIndexByWidgetMap.Find(Child))
+	int32 PreviousEntryIndex = Entries.IndexOfByPredicate([Child](const FUIFrameworkWidgetTreeEntry& Other){ return Other.Child == Child; });
+	if (PreviousEntryIndex != INDEX_NONE)
 	{
-		check(Entries.IsValidIndex(*PreviousEntryIndexPtr));
-		FUIFrameworkWidgetTreeEntry& PreviousEntry = Entries[*PreviousEntryIndexPtr];
+		const FUIFrameworkWidgetTreeEntry& PreviousEntry = Entries[PreviousEntryIndex];
 		if (PreviousEntry.Parent != Parent)
 		{
-			// Same child, different parent. Need to build a new entry for replication.
-			PreviousEntry = FUIFrameworkWidgetTreeEntry(Parent, Child);
-			MarkItemDirty(PreviousEntry);
+			// New Parent for an existing child.
+			Entries.RemoveAtSwap(PreviousEntryIndex);
+			FUIFrameworkWidgetTreeEntry& NewEntry = Entries.Emplace_GetRef(Parent, Child);
+			MarkItemDirty(NewEntry);
+		}
+		else
+		{
+			UE_LOG(LogUIFramework, Warning, TEXT("A widget was added to the WidgetTree but it's already added."));
 		}
 	}
 	else
 	{
-		int32 NewEntryIndex = Entries.Emplace(Parent, Child);
-
-		FUIFrameworkWidgetTreeEntry& NewEntry = Entries[NewEntryIndex];
+		FUIFrameworkWidgetTreeEntry& NewEntry = Entries.Emplace_GetRef(Parent, Child);
 		MarkItemDirty(NewEntry);
-
-		AuthorityIndexByWidgetMap.Add(Child) = NewEntryIndex;
-		WidgetByIdMap.Add(Child->GetWidgetId()) = Child;
+		WidgetByIdMap.FindOrAdd(Child->GetWidgetId()) = Child;
 
 		if (ensure(ReplicatedOwner) && ReplicatedOwner->IsUsingRegisteredSubObjectList())
 		{
 			ReplicatedOwner->AddReplicatedSubObject(Child);
 		}
-
-		AuthorityAddChildRecursiveInternal(Child);
+		AddChildRecursiveInternal(Child);
 	}
 }
 
-void FUIFrameworkWidgetTree::AuthorityAddChildRecursiveInternal(UUIFrameworkWidget* InParentWidget)
+void FUIFrameworkWidgetTree::AddChildRecursiveInternal(UUIFrameworkWidget* Widget)
 {
 	FUIFrameworkWidgetTree* Self = this;
-	InParentWidget->AuthorityForEachChildren([Self, InParentWidget](UUIFrameworkWidget* ChildWidget)
+	Widget->AuthorityForEachChildren([Self, Widget](UUIFrameworkWidget* ChildWidget)
 		{
 			if (ChildWidget != nullptr)
 			{
-				Self->AuthorityAddChildInternal(InParentWidget, ChildWidget);
+				FUIFrameworkWidgetTreeEntry& NewEntry = Self->Entries.Emplace_GetRef(Widget, ChildWidget);
+				Self->MarkItemDirty(NewEntry);
+				Self->WidgetByIdMap.FindOrAdd(ChildWidget->GetWidgetId()) = ChildWidget;
+				if (ensure(Self->ReplicatedOwner) && Self->ReplicatedOwner->IsUsingRegisteredSubObjectList())
+				{
+					Self->ReplicatedOwner->AddReplicatedSubObject(ChildWidget);
+				}
+				Self->AddChildRecursiveInternal(ChildWidget);
 			}
 		});
 }
 
-void FUIFrameworkWidgetTree::AuthorityRemoveWidget(UUIFrameworkWidget* Widget)
+void FUIFrameworkWidgetTree::RemoveWidget(UUIFrameworkWidget* Widget)
 {
 	check(Widget);
-
-	if (AuthorityRemoveChildRecursiveInternal(Widget))
-	{
-		MarkArrayDirty();
-	}
-}
-
-bool FUIFrameworkWidgetTree::AuthorityRemoveChildRecursiveInternal(UUIFrameworkWidget* Widget)
-{
-	if (int32* PreviousEntryIndexPtr = AuthorityIndexByWidgetMap.Find(Widget))
+	bool bDirtyArray = false;
 	{
-		check(Entries.IsValidIndex(*PreviousEntryIndexPtr));
-
-		AuthorityIndexByWidgetMap.Remove(Widget);
-		WidgetByIdMap.Remove(Widget->GetWidgetId());
-
-		if (ensure(ReplicatedOwner) && ReplicatedOwner->IsUsingRegisteredSubObjectList())
-		{
-			ReplicatedOwner->RemoveReplicatedSubObject(Widget);
-		}
+		FMemMark Mark(FMemStack::Get());
+		TArray<UUIFrameworkWidget*, TInlineAllocator<60, TMemStackAllocator<>>> WidgetsToRemove;
+		WidgetsToRemove.Add(Widget);
 
-		Entries.RemoveAtSwap(*PreviousEntryIndexPtr);
-
-		// Fix up the swap item
-		if (Entries.IsValidIndex(*PreviousEntryIndexPtr))
+		while (WidgetsToRemove.Num() > 0)
 		{
-			if (int32* FixUpChildIndexPtr = AuthorityIndexByWidgetMap.Find(Entries[*PreviousEntryIndexPtr].Child))
+			UUIFrameworkWidget* CurrentWidgetToRemove = WidgetsToRemove.Last();
+			WidgetByIdMap.Remove(CurrentWidgetToRemove->GetWidgetId());
+			if (ensure(ReplicatedOwner) && ReplicatedOwner->IsUsingRegisteredSubObjectList())
 			{
-				*FixUpChildIndexPtr = *PreviousEntryIndexPtr;
+				ReplicatedOwner->RemoveReplicatedSubObject(CurrentWidgetToRemove);
 			}
-		}
 
-		FUIFrameworkWidgetTree* Self = this;
-		Widget->AuthorityForEachChildren([Self](UUIFrameworkWidget* ChildWidget)
+			WidgetsToRemove.RemoveAt(WidgetsToRemove.Num()-1);
+			for (int32 Index = Entries.Num() - 1; Index >= 0; --Index)
 			{
-				if (ChildWidget)
+				if (Entries[Index].Parent == CurrentWidgetToRemove)
 				{
-					Self->AuthorityRemoveChildRecursiveInternal(ChildWidget);
-				}
-			});	
+					bDirtyArray = true;
+					if (Entries[Index].Child)
+					{
+						WidgetsToRemove.Add(Entries[Index].Child);
+					}
 
-		return true;
+					Entries.RemoveAtSwap(Index);
+				}
+				else if (Entries[Index].Child == CurrentWidgetToRemove)
+				{
+					bDirtyArray = true;
+					Entries.RemoveAtSwap(Index);
+				}
+			}
+		}
 	}
-	return false;
 
+	if (bDirtyArray)
+	{
+		MarkArrayDirty();
+	}
 }
 
-FUIFrameworkWidgetTreeEntry* FUIFrameworkWidgetTree::LocalGetEntryByReplicationId(int32 ReplicationId)
+FUIFrameworkWidgetTreeEntry* FUIFrameworkWidgetTree::GetEntryByReplicationId(int32 ReplicationId)
 {
 	if (const int32* Index = ItemMap.Find(ReplicationId))
 	{
@@ -271,9 +257,9 @@ FUIFrameworkWidgetTreeEntry* FUIFrameworkWidgetTree::LocalGetEntryByReplicationI
 	return nullptr;
 }
 
-const FUIFrameworkWidgetTreeEntry* FUIFrameworkWidgetTree::LocalGetEntryByReplicationId(int32 ReplicationId) const
+const FUIFrameworkWidgetTreeEntry* FUIFrameworkWidgetTree::GetEntryByReplicationId(int32 ReplicationId) const
 {
-	return const_cast<FUIFrameworkWidgetTree*>(this)->LocalGetEntryByReplicationId(ReplicationId);
+	return const_cast<FUIFrameworkWidgetTree*>(this)->GetEntryByReplicationId(ReplicationId);
 }
 
 UUIFrameworkWidget* FUIFrameworkWidgetTree::FindWidgetById(FUIFrameworkWidgetId WidgetId)
@@ -285,86 +271,4 @@ UUIFrameworkWidget* FUIFrameworkWidgetTree::FindWidgetById(FUIFrameworkWidgetId
 const UUIFrameworkWidget* FUIFrameworkWidgetTree::FindWidgetById(FUIFrameworkWidgetId WidgetId) const
 {
 	return const_cast<FUIFrameworkWidgetTree*>(this)->FindWidgetById(WidgetId);
-}
-
-#if UE_UIFRAMEWORK_WITH_DEBUG
-void FUIFrameworkWidgetTree::AuthorityTest() const
-{
-	if (ReplicatedOwner == nullptr || !ReplicatedOwner->HasAuthority())
-	{
-		return;
-	}
-
-	TSet<FUIFrameworkWidgetId> UniqueIds;
-	for (int32 Index = 0; Index < Entries.Num(); ++Index)
-	{
-		const FUIFrameworkWidgetTreeEntry& Entry = Entries[Index];
-
-		ensureAlwaysMsgf(Entry.ParentId.IsValid(), TEXT("Invalid ParentId"));
-		if (!Entry.ParentId.IsRoot())
-		{
-			ensureAlwaysMsgf(Entry.Parent, TEXT("Invalid Parent"));
-		}
-		else
-		{
-			ensureAlwaysMsgf(!Entry.Parent, TEXT("Valid Parent"));
-		}
-		ensureAlwaysMsgf(Entry.Child, TEXT("Invalid Child"));
-		ensureAlwaysMsgf(Entry.ChildId.IsValid(), TEXT("Invalid ChildId"));
-
-		if (Entry.ChildId.IsValid())
-		{
-			ensureAlwaysMsgf(!UniqueIds.Contains(Entry.ChildId), TEXT("Duplicated id"));
-			UniqueIds.Add(Entry.ChildId);
-		}
-
-		if (Entry.Child)
-		{
-			ensureAlwaysMsgf(Entry.ChildId == Entry.Child->GetWidgetId(), TEXT("Id do not matches"));
-		}
-		
-		if (Entry.Parent)
-		{
-			bool bFound = false;
-			const UUIFrameworkWidget* ToFindWidget = Entry.Child;
-			Entry.Parent->AuthorityForEachChildren([&bFound, ToFindWidget](UUIFrameworkWidget* ChildWidget)
-				{
-					bFound = bFound || ToFindWidget == ChildWidget;
-				});
-
-			ensureAlwaysMsgf(bFound, TEXT("Widget is in the tree but not in the AuthorityForEachChildren"));
-		}
-
-
-		const int32* FoundIndexPtr = AuthorityIndexByWidgetMap.Find(Entry.Child);
-		if (FoundIndexPtr)
-		{
-			ensureAlwaysMsgf(*FoundIndexPtr == Index, TEXT("Widget index doesn't match what is in the map"));
-		}
-		else
-		{
-			ensureAlwaysMsgf(false, TEXT("Widget no in the map"));
-		}
-		
-
-		if (Entry.Child)
-		{
-			ensureAlwaysMsgf(Entry.Child->GetWidgetId().IsValid(), TEXT("The id is not valid."));
-			const TWeakObjectPtr<UUIFrameworkWidget>* FoundWidgetPtr = WidgetByIdMap.Find(Entry.ChildId);
-			if (FoundWidgetPtr)
-			{
-				ensureAlwaysMsgf(FoundWidgetPtr->Get(), TEXT("The found widget is invalid"));
-				if (UUIFrameworkWidget* Widget = FoundWidgetPtr->Get())
-				{
-					ensureAlwaysMsgf(Widget == Entry.Child, TEXT("Widget in the map doesn't matches with the entry widget."));
-				}
-			}
-			else
-			{
-				ensureAlwaysMsgf(false, TEXT("Widget no in the map"));
-			}
-		}
-		
-	}
-}
-#endif
\ No newline at end of file
+}
\ No newline at end of file
diff --git a/Engine/Plugins/Experimental/UIFramework/Source/Private/UIFModule.cpp b/Engine/Plugins/Experimental/UIFramework/Source/Private/UIFModule.cpp
deleted file mode 100644
diff --git a/Engine/Plugins/Experimental/UIFramework/Source/Private/UIFPlayerComponent.cpp b/Engine/Plugins/Experimental/UIFramework/Source/Private/UIFPlayerComponent.cpp
--- a/Engine/Plugins/Experimental/UIFramework/Source/Private/UIFPlayerComponent.cpp
+++ b/Engine/Plugins/Experimental/UIFramework/Source/Private/UIFPlayerComponent.cpp
@@ -2,7 +2,6 @@
 
 #include "UIFPlayerComponent.h"
 #include "UIFLog.h"
-#include "UIFModule.h"
 #include "UIFWidget.h"
 #include "Blueprint/GameViewportSubsystem.h"
 
@@ -122,8 +121,8 @@ void UUIFrameworkPlayerComponent::AddWidget(FUIFrameworkGameLayerSlot InEntry)
 		}
 		else
 		{
-			// Reset the widget to make sure the id is set and it may have been duplicated during the attach
-			InEntry.AuthoritySetWidget(FUIFrameworkModule::AuthorityAttachWidget(this, this, InEntry.AuthorityGetWidget()));
+			InEntry.AuthoritySetWidget(InEntry.AuthorityGetWidget()); // to make sure the id is set
+			InEntry.AuthorityGetWidget()->AuthoritySetParent(this, FUIFrameworkParentWidget(this));
 			RootList.AddEntry(InEntry);
 		}
 	}
@@ -147,8 +146,8 @@ void UUIFrameworkPlayerComponent::RemoveWidget(UUIFrameworkWidget* Widget)
 		}
 		else
 		{
-			FUIFrameworkModule::AuthorityDetachWidgetFromParent(Widget);
 			RootList.RemoveEntry(Widget);
+			Widget->AuthoritySetParent(this, FUIFrameworkParentWidget());
 		}
 	}
 }
@@ -171,7 +170,7 @@ void UUIFrameworkPlayerComponent::TickComponent(float DeltaTime, enum ELevelTick
 		TGuardValue<bool> TmpGuard = {bAddingWidget, true};
 		for (int32 ReplicationId : AddPending)
 		{
-			if (FUIFrameworkWidgetTreeEntry* Entry = WidgetTree.LocalGetEntryByReplicationId(ReplicationId))
+			if (FUIFrameworkWidgetTreeEntry* Entry = WidgetTree.GetEntryByReplicationId(ReplicationId))
 			{
 				if (ensure(Entry->IsParentValid() && Entry->IsChildValid()))
 				{
@@ -210,7 +209,7 @@ void UUIFrameworkPlayerComponent::TickComponent(float DeltaTime, enum ELevelTick
 
 		for (int32 ReplicationId : TempNetReplicationPending)
 		{
-			if (const FUIFrameworkWidgetTreeEntry* Entry = WidgetTree.LocalGetEntryByReplicationId(ReplicationId))
+			if (const FUIFrameworkWidgetTreeEntry* Entry = WidgetTree.GetEntryByReplicationId(ReplicationId))
 			{
 				LocalWidgetWasAddedToTree(*Entry);
 			}
@@ -300,7 +299,7 @@ void UUIFrameworkPlayerComponent::LocalOnClassLoaded(TSoftClassPtr<UWidget> Widg
 		{
 			for (int32 ReplicationId : FoundWidgetClassToLoad.EntryReplicationIds)
 			{
-				if (const FUIFrameworkWidgetTreeEntry* Entry = WidgetTree.LocalGetEntryByReplicationId(ReplicationId))
+				if (const FUIFrameworkWidgetTreeEntry* Entry = WidgetTree.GetEntryByReplicationId(ReplicationId))
 				{
 					if (Entry->IsParentValid() && Entry->IsChildValid())
 					{
diff --git a/Engine/Plugins/Experimental/UIFramework/Source/Private/UIFWidget.cpp b/Engine/Plugins/Experimental/UIFramework/Source/Private/UIFWidget.cpp
--- a/Engine/Plugins/Experimental/UIFramework/Source/Private/UIFWidget.cpp
+++ b/Engine/Plugins/Experimental/UIFramework/Source/Private/UIFWidget.cpp
@@ -66,10 +66,74 @@ void UUIFrameworkWidget::GetLifetimeReplicatedProps(TArray<FLifetimeProperty>& O
 	DOREPLIFETIME_WITH_PARAMS_FAST(ThisClass, Id, Params);
 }
 
+void UUIFrameworkWidget::AuthoritySetParent(UUIFrameworkPlayerComponent* NewOwner, FUIFrameworkParentWidget NewParent)
+{
+	const bool bDifferentOwner = NewOwner != OwnerPlayerComponent;
+	if (OwnerPlayerComponent)
+	{
+		ensure(OwnerPlayerComponent == NewOwner);
+		NewOwner = OwnerPlayerComponent;
+		NewParent = FUIFrameworkParentWidget();
+	}
+
+	if (AuthorityParent.IsParentValid())
+	{
+		if (AuthorityParent.IsWidget())
+		{
+			AuthorityParent.AsWidget()->AuthorityRemoveChild(this);
+		}
+		else
+		{
+			check(AuthorityParent.IsPlayerComponent());
+			AuthorityParent.AsPlayerComponent()->AuthorityRemoveChild(this);
+		}
+	}
+
+	AuthorityParent = NewParent;
+	OwnerPlayerComponent = NewOwner;
+
+	if (AuthorityParent.IsParentValid() && OwnerPlayerComponent)
+	{
+		if (AuthorityParent.IsWidget())
+		{
+			OwnerPlayerComponent->GetWidgetTree().AddWidget(AuthorityParent.AsWidget(), this);
+		}
+		else
+		{
+			check(AuthorityParent.IsPlayerComponent());
+			OwnerPlayerComponent->GetWidgetTree().AddRoot(this);
+		}
+	}
+	else if (OwnerPlayerComponent)
+	{
+		OwnerPlayerComponent->GetWidgetTree().RemoveWidget(this);
+	}
+
+	if (bDifferentOwner)
+	{
+		SetParentPlayerOwnerRecursive();
+	}
+}
+
+void UUIFrameworkWidget::SetParentPlayerOwnerRecursive()
+{
+	UUIFrameworkWidget* Self = this;
+	AuthorityForEachChildren([Self](UUIFrameworkWidget* Child)
+		{
+			if (Child != nullptr)
+			{
+				check(Child->AuthorityGetParent().IsWidget() && Child->AuthorityGetParent().AsWidget() == Self);
+				Child->OwnerPlayerComponent = Self->OwnerPlayerComponent;
+				Child->AuthorityParent = FUIFrameworkParentWidget(Self);
+				Child->SetParentPlayerOwnerRecursive();
+			}
+		});
+}
+
 void UUIFrameworkWidget::LocalAddChild(FUIFrameworkWidgetId ChildId)
 {
 	// By default we should remove the widget from its previous parent.
-	//Adding a widget to a new slot will automatically remove it from its previous parent.
+	//Adding a widget to a new slot will automaticly remove it from its previous parent.
 	if (OwnerPlayerComponent)
 	{
 		if (UUIFrameworkWidget* Widget = OwnerPlayerComponent->GetWidgetTree().FindWidgetById(ChildId))
diff --git a/Engine/Plugins/Experimental/UIFramework/Source/Private/UIFrameworkModule.cpp b/Engine/Plugins/Experimental/UIFramework/Source/Private/UIFrameworkModule.cpp
new file mode 100644
--- /dev/null
+++ b/Engine/Plugins/Experimental/UIFramework/Source/Private/UIFrameworkModule.cpp
@@ -0,0 +1,5 @@
+// Copyright Epic Games, Inc. All Rights Reserved.
+
+#include "Modules/ModuleManager.h"
+	
+IMPLEMENT_MODULE(FDefaultModuleImpl, UIFramework)
\ No newline at end of file
diff --git a/Engine/Plugins/Experimental/UIFramework/Source/Private/Widgets/UIFButton.cpp b/Engine/Plugins/Experimental/UIFramework/Source/Private/Widgets/UIFButton.cpp
--- a/Engine/Plugins/Experimental/UIFramework/Source/Private/Widgets/UIFButton.cpp
+++ b/Engine/Plugins/Experimental/UIFramework/Source/Private/Widgets/UIFButton.cpp
@@ -3,7 +3,6 @@
 #include "Widgets/UIFButton.h"
 #include "Types/UIFWidgetTree.h"
 #include "UIFLog.h"
-#include "UIFModule.h"
 #include "UIFPlayerComponent.h"
 
 #include "Components/Button.h"
@@ -42,7 +41,7 @@ void UUIFrameworkButton::SetContent(FUIFrameworkSimpleSlot InEntry)
 		// Remove previous widget
 		if (Slot.AuthorityGetWidget())
 		{
-			FUIFrameworkModule::AuthorityDetachWidgetFromParent(Slot.AuthorityGetWidget());
+			Slot.AuthorityGetWidget()->AuthoritySetParent(GetPlayerComponent(), FUIFrameworkParentWidget());
 		}
 
 		if (InEntry.AuthorityGetWidget())
@@ -57,11 +56,11 @@ void UUIFrameworkButton::SetContent(FUIFrameworkSimpleSlot InEntry)
 	}
 
 	Slot = InEntry;
+	Slot.AuthoritySetWidget(InEntry.AuthorityGetWidget()); // to make sure the id is set
 
 	if (bWidgetIsDifferent && Slot.AuthorityGetWidget())
 	{
-		// Reset the widget to make sure the id is set and it may have been duplicated during the attach
-		Slot.AuthoritySetWidget(FUIFrameworkModule::AuthorityAttachWidget(GetPlayerComponent(), this, Slot.AuthorityGetWidget()));
+		Slot.AuthorityGetWidget()->AuthoritySetParent(GetPlayerComponent(), FUIFrameworkParentWidget(this));
 	}
 
 	MARK_PROPERTY_DIRTY_FROM_NAME(ThisClass, Slot, this);
diff --git a/Engine/Plugins/Experimental/UIFramework/Source/Private/Widgets/UIFCanvasBox.cpp b/Engine/Plugins/Experimental/UIFramework/Source/Private/Widgets/UIFCanvasBox.cpp
--- a/Engine/Plugins/Experimental/UIFramework/Source/Private/Widgets/UIFCanvasBox.cpp
+++ b/Engine/Plugins/Experimental/UIFramework/Source/Private/Widgets/UIFCanvasBox.cpp
@@ -3,7 +3,6 @@
 #include "Widgets/UIFCanvasBox.h"
 #include "Types/UIFWidgetTree.h"
 #include "UIFLog.h"
-#include "UIFModule.h"
 #include "UIFPlayerComponent.h"
 
 #include "Components/CanvasPanel.h"
@@ -106,8 +105,8 @@ void UUIFrameworkCanvasBox::AddWidget(FUIFrameworkCanvasBoxSlot InEntry)
 	}
 	else
 	{
-		// Reset the widget to make sure the id is set and it may have been duplicated during the attach
-		InEntry.AuthoritySetWidget(FUIFrameworkModule::AuthorityAttachWidget(GetPlayerComponent(), this, InEntry.AuthorityGetWidget()));
+		InEntry.AuthoritySetWidget(InEntry.AuthorityGetWidget()); // to make sure the id is set
+		InEntry.AuthorityGetWidget()->AuthoritySetParent(GetPlayerComponent(), FUIFrameworkParentWidget(this));
 		AddEntry(InEntry);
 	}
 }
@@ -127,8 +126,8 @@ void UUIFrameworkCanvasBox::RemoveWidget(UUIFrameworkWidget* Widget)
 		}
 		else
 		{
-			FUIFrameworkModule::AuthorityDetachWidgetFromParent(Widget);
 			RemoveEntry(Widget);
+			Widget->AuthoritySetParent(GetPlayerComponent(), FUIFrameworkParentWidget());
 		}
 	}
 }
diff --git a/Engine/Plugins/Experimental/UIFramework/Source/Private/Widgets/UIFStackBox.cpp b/Engine/Plugins/Experimental/UIFramework/Source/Private/Widgets/UIFStackBox.cpp
--- a/Engine/Plugins/Experimental/UIFramework/Source/Private/Widgets/UIFStackBox.cpp
+++ b/Engine/Plugins/Experimental/UIFramework/Source/Private/Widgets/UIFStackBox.cpp
@@ -3,7 +3,6 @@
 #include "Widgets/UIFStackBox.h"
 #include "Types/UIFWidgetTree.h"
 #include "UIFLog.h"
-#include "UIFModule.h"
 #include "UIFPlayerComponent.h"
 
 #include "Components/StackBox.h"
@@ -43,23 +42,17 @@ bool FUIFrameworkStackBoxSlotList::NetDeltaSerialize(FNetDeltaSerializeInfo& Del
 
 void FUIFrameworkStackBoxSlotList::AddEntry(FUIFrameworkStackBoxSlot Entry)
 {
-	int32 NewEntryIndex = Slots.Add(MoveTemp(Entry));
-	FUIFrameworkStackBoxSlot& NewEntry = Slots[NewEntryIndex];
-	NewEntry.Index = NewEntryIndex;
+	FUIFrameworkStackBoxSlot& NewEntry = Slots.Add_GetRef(MoveTemp(Entry));
 	MarkItemDirty(NewEntry);
 }
 
 bool FUIFrameworkStackBoxSlotList::RemoveEntry(UUIFrameworkWidget* Widget)
 {
 	check(Widget);
-	int32 Index = Slots.IndexOfByPredicate([Widget](const FUIFrameworkStackBoxSlot& Entry) { return Entry.AuthorityGetWidget() == Widget; });
+	const int32 Index = Slots.IndexOfByPredicate([Widget](const FUIFrameworkStackBoxSlot& Entry) { return Entry.AuthorityGetWidget() == Widget; });
 	if (Index != INDEX_NONE)
 	{
 		Slots.RemoveAt(Index);
-		for (; Index < Slots.Num(); ++Index)
-		{
-			Slots[Index].Index = Index;
-		}
 		MarkArrayDirty();
 	}
 	return Index != INDEX_NONE;
@@ -110,8 +103,8 @@ void UUIFrameworkStackBox::AddWidget(FUIFrameworkStackBoxSlot InEntry)
 	}
 	else
 	{
-		// Reset the widget to make sure the id is set and it may have been duplicated during the attach
-		InEntry.AuthoritySetWidget(FUIFrameworkModule::AuthorityAttachWidget(GetPlayerComponent(), this, InEntry.AuthorityGetWidget()));
+		InEntry.AuthoritySetWidget(InEntry.AuthorityGetWidget()); // to make sure the id is set
+		InEntry.AuthorityGetWidget()->AuthoritySetParent(GetPlayerComponent(), FUIFrameworkParentWidget(this));
 		AddEntry(InEntry);
 	}
 }
@@ -131,8 +124,8 @@ void UUIFrameworkStackBox::RemoveWidget(UUIFrameworkWidget* Widget)
 		}
 		else
 		{
-			FUIFrameworkModule::AuthorityDetachWidgetFromParent(Widget);
 			RemoveEntry(Widget);
+			Widget->AuthoritySetParent(GetPlayerComponent(), FUIFrameworkParentWidget());
 		}
 	}
 }
diff --git a/Engine/Plugins/Experimental/UIFramework/Source/Public/Types/UIFWidgetTree.h b/Engine/Plugins/Experimental/UIFramework/Source/Public/Types/UIFWidgetTree.h
--- a/Engine/Plugins/Experimental/UIFramework/Source/Public/Types/UIFWidgetTree.h
+++ b/Engine/Plugins/Experimental/UIFramework/Source/Public/Types/UIFWidgetTree.h
@@ -9,10 +9,6 @@
 
 #include "UIFWidgetTree.generated.h"
 
-#ifndef UE_UIFRAMEWORK_WITH_DEBUG
-	#define UE_UIFRAMEWORK_WITH_DEBUG !(UE_BUILD_SHIPPING || UE_BUILD_TEST)
-#endif
-
 struct FReplicationFlags;
 struct FUIFrameworkWidgetTreeEntry;
 class AActor;
@@ -69,8 +65,11 @@ struct UIFRAMEWORK_API FUIFrameworkWidgetTree : public FFastArraySerializer
 
 public:
 	FUIFrameworkWidgetTree() = default;
-	FUIFrameworkWidgetTree(AActor* InReplicatedOwner, IUIFrameworkWidgetTreeOwner* InOwner);
-	~FUIFrameworkWidgetTree();
+	FUIFrameworkWidgetTree(AActor* InReplicatedOwner, IUIFrameworkWidgetTreeOwner* InOwner)
+		: ReplicatedOwner(InReplicatedOwner)
+		, Owner(InOwner)
+	{
+	}
 
 public:
 	//~ Begin of FFastArraySerializer
@@ -87,44 +86,39 @@ public:
 	bool ReplicateSubWidgets(UActorChannel* Channel, FOutBunch* Bunch, FReplicationFlags* RepFlags);
 
 	/** Add a new widget to the top hierarchy. */
-	void AuthorityAddRoot(UUIFrameworkWidget* Widget);
+	void AddRoot(UUIFrameworkWidget* Widget);
 	/**
 	 * Change the parent / child relationship of the child widget.
 	 * If the child widget had a parent, that relationship entry will replaced it by a new one.
 	 */
-	void AuthorityAddWidget(UUIFrameworkWidget* Parent, UUIFrameworkWidget* Child);
+	void AddWidget(UUIFrameworkWidget* Parent, UUIFrameworkWidget* Child);
 	/**
 	 * Remove the widget and all of its children and grand-children from the tree.
 	 * It will clean all the parent relationship from the tree.
 	 */
-	void AuthorityRemoveWidget(UUIFrameworkWidget* Widget);
+	void RemoveWidget(UUIFrameworkWidget* Widget);
 
-	//~ It is not safe to use the ReplicationId on the Authority because any add or remove would clear the ItemMap.
-	FUIFrameworkWidgetTreeEntry* LocalGetEntryByReplicationId(int32 WidgetId);
-	const FUIFrameworkWidgetTreeEntry* LocalGetEntryByReplicationId(int32 WidgetId) const;
+	FUIFrameworkWidgetTreeEntry* GetEntryByReplicationId(int32 ReplicationId);
+	const FUIFrameworkWidgetTreeEntry* GetEntryByReplicationId(int32 ReplicationId) const;
 
 	/** Find the widget by its unique Id. The widget needs to be in the Tree. */
 	UUIFrameworkWidget* FindWidgetById(FUIFrameworkWidgetId WidgetId);
 	const UUIFrameworkWidget* FindWidgetById(FUIFrameworkWidgetId WidgetId) const;
 
-#if UE_UIFRAMEWORK_WITH_DEBUG
-	void AuthorityTest() const;
-#endif
-
 private:
-	void AuthorityAddChildInternal(UUIFrameworkWidget* Parent, UUIFrameworkWidget* Child);
-	void AuthorityAddChildRecursiveInternal(UUIFrameworkWidget* Widget);
-	bool AuthorityRemoveChildRecursiveInternal(UUIFrameworkWidget* Widget);
+	void AddChildInternal(UUIFrameworkWidget* Parent, UUIFrameworkWidget* Child);
+	void AddChildRecursiveInternal(UUIFrameworkWidget* Widget);
 
 private:
 	UPROPERTY()
 	TArray<FUIFrameworkWidgetTreeEntry> Entries;
 
+	UPROPERTY(NotReplicated)
+	TMap<FUIFrameworkWidgetId, TWeakObjectPtr<UUIFrameworkWidget>> WidgetByIdMap;
+
 	UPROPERTY(NotReplicated, Transient)
 	TObjectPtr<AActor> ReplicatedOwner;
 
-	TMap<TWeakObjectPtr<UUIFrameworkWidget>, int32> AuthorityIndexByWidgetMap;
-	TMap<FUIFrameworkWidgetId, TWeakObjectPtr<UUIFrameworkWidget>> WidgetByIdMap;
 	IUIFrameworkWidgetTreeOwner* Owner = nullptr;
 };
 
diff --git a/Engine/Plugins/Experimental/UIFramework/Source/Public/UIFModule.h b/Engine/Plugins/Experimental/UIFramework/Source/Public/UIFModule.h
deleted file mode 100644
diff --git a/Engine/Plugins/Experimental/UIFramework/Source/Public/UIFWidget.h b/Engine/Plugins/Experimental/UIFramework/Source/Public/UIFWidget.h
--- a/Engine/Plugins/Experimental/UIFramework/Source/Public/UIFWidget.h
+++ b/Engine/Plugins/Experimental/UIFramework/Source/Public/UIFWidget.h
@@ -9,7 +9,6 @@
 
 #include "UIFWidget.generated.h"
 
-class FUIFrameworkModule;
 
 /**
  * 
@@ -18,7 +17,6 @@ UCLASS(Abstract, BlueprintType)
 class UIFRAMEWORK_API UUIFrameworkWidget : public UObject
 {
 	GENERATED_BODY()
-	friend FUIFrameworkModule;
 
 public:
 	//~ Begin UObject
@@ -46,6 +44,8 @@ public:
 	}
 
 	//~ Authority functions
+	void AuthoritySetParent(UUIFrameworkPlayerComponent* Owner, FUIFrameworkParentWidget NewParent);
+
 	FUIFrameworkParentWidget AuthorityGetParent() const
 	{
 		return AuthorityParent;
@@ -67,6 +67,9 @@ protected:
 	virtual void AuthorityRemoveChild(UUIFrameworkWidget* Widget) {}
 	virtual void LocalOnUMGWidgetCreated() { }
 
+private:
+	void SetParentPlayerOwnerRecursive();
+
 protected:
 	UPROPERTY(BlueprintReadOnly, EditDefaultsOnly, Category = "UI Framework")
 	TSoftClassPtr<UWidget> WidgetClass; // todo: make this private and use a constructor argument
diff --git a/Engine/Plugins/Experimental/UIFramework/Source/Public/Widgets/UIFStackBox.h b/Engine/Plugins/Experimental/UIFramework/Source/Public/Widgets/UIFStackBox.h
--- a/Engine/Plugins/Experimental/UIFramework/Source/Public/Widgets/UIFStackBox.h
+++ b/Engine/Plugins/Experimental/UIFramework/Source/Public/Widgets/UIFStackBox.h
@@ -21,8 +21,6 @@ struct FUIFrameworkStackBoxSlot : public FUIFrameworkSlotBase
 {
 	GENERATED_BODY()
 
-	friend FUIFrameworkStackBoxSlotList;
-
 	/** Horizontal alignment of the widget inside the slot. */
 	UPROPERTY(EditAnywhere, BlueprintReadWrite, Category = "UI Framework")
 	TEnumAsByte<EHorizontalAlignment> HorizontalAlignment = EHorizontalAlignment::HAlign_Fill;
@@ -38,11 +36,6 @@ struct FUIFrameworkStackBoxSlot : public FUIFrameworkSlotBase
 	/** How much space this slot should occupy in the direction of the panel. */
 	UPROPERTY(EditAnywhere, BlueprintReadWrite, Category = "UI Framework")
 	FSlateChildSize Size = FSlateChildSize(ESlateSizeRule::Automatic);
-
-private:
-	/** Index in the array the Slot is. The position in the array can change when replicated. */
-	UPROPERTY()
-	int32 Index;
 };
 
 
