Add text widget to the ui framework

diff --git a/Engine/Plugins/Experimental/UIFramework/Source/Public/Types/UIFEvents.h b/Engine/Plugins/Experimental/UIFramework/Source/Public/Types/UIFEvents.h
--- a/Engine/Plugins/Experimental/UIFramework/Source/Public/Types/UIFEvents.h
+++ b/Engine/Plugins/Experimental/UIFramework/Source/Public/Types/UIFEvents.h
@@ -14,11 +14,11 @@ class UUIFrameworkWidget;
  *
  */
 USTRUCT(BlueprintType)
-struct FUIFrameworkClickEventArgument
+struct FUIFrameworkSimpleEventArgument
 {
 	GENERATED_BODY()
 
-	FUIFrameworkClickEventArgument() = default;
+	FUIFrameworkSimpleEventArgument() = default;
 
 	UPROPERTY()
 	TObjectPtr<APlayerController> PlayerController;
@@ -27,4 +27,17 @@ struct FUIFrameworkClickEventArgument
 	TObjectPtr<UUIFrameworkWidget> Sender;
 };
 
+DECLARE_MULTICAST_DELEGATE_OneParam(FUIFrameworkSimpleEvent, FUIFrameworkSimpleEventArgument);
+
+/**
+ *
+ */
+USTRUCT(BlueprintType)
+struct FUIFrameworkClickEventArgument : public FUIFrameworkSimpleEventArgument
+{
+	GENERATED_BODY()
+
+	FUIFrameworkClickEventArgument() = default;
+};
+
 DECLARE_MULTICAST_DELEGATE_OneParam(FUIFrameworkClickEvent, FUIFrameworkClickEventArgument);
