Fix static analyzer warning in UIFramework

diff --git a/Engine/Plugins/Experimental/UIFramework/Source/Private/Types/UIFWidgetTree.cpp b/Engine/Plugins/Experimental/UIFramework/Source/Private/Types/UIFWidgetTree.cpp
--- a/Engine/Plugins/Experimental/UIFramework/Source/Private/Types/UIFWidgetTree.cpp
+++ b/Engine/Plugins/Experimental/UIFramework/Source/Private/Types/UIFWidgetTree.cpp
@@ -339,7 +339,7 @@ void FUIFrameworkWidgetTree::AuthorityTest() const
 		const int32* FoundIndexPtr = AuthorityIndexByWidgetMap.Find(Entry.Child);
 		if (FoundIndexPtr)
 		{
-			ensureAlwaysMsgf(*FoundIndexPtr == Index, TEXT("Widget index doesn't match what is in the map"));
+			ensureAlwaysMsgf((*FoundIndexPtr) == Index, TEXT("Widget index doesn't match what is in the map"));
 		}
 		else
 		{
