Visibility and Enable of UIFrameworkWidget.

diff --git a/Engine/Plugins/Experimental/UIFramework/Source/Private/UIFWidget.cpp b/Engine/Plugins/Experimental/UIFramework/Source/Private/UIFWidget.cpp
--- a/Engine/Plugins/Experimental/UIFramework/Source/Private/UIFWidget.cpp
+++ b/Engine/Plugins/Experimental/UIFramework/Source/Private/UIFWidget.cpp
@@ -64,6 +64,8 @@ void UUIFrameworkWidget::GetLifetimeReplicatedProps(TArray<FLifetimeProperty>& O
 	FDoRepLifetimeParams Params;
 	Params.bIsPushBased = true;
 	DOREPLIFETIME_WITH_PARAMS_FAST(ThisClass, Id, Params);
+	DOREPLIFETIME_WITH_PARAMS_FAST(ThisClass, bIsEnabled, Params);
+	DOREPLIFETIME_WITH_PARAMS_FAST(ThisClass, Visibility, Params);
 }
 
 void UUIFrameworkWidget::LocalAddChild(FUIFrameworkWidgetId ChildId)
@@ -96,6 +98,8 @@ void UUIFrameworkWidget::LocalCreateUMGWidget(UUIFrameworkPlayerComponent* InOwn
 			check(Class->IsChildOf(UWidget::StaticClass()));
 			LocalUMGWidget = NewObject<UWidget>(this, Class, FName(), RF_Transient);
 		}
+		LocalUMGWidget->SetIsEnabled(bIsEnabled);
+		LocalUMGWidget->SetVisibility(Visibility);
 		LocalOnUMGWidgetCreated();
 	}
 }
@@ -110,3 +114,48 @@ void UUIFrameworkWidget::LocalDestroyUMGWidget()
 	LocalUMGWidget = nullptr;
 	OwnerPlayerComponent = nullptr;
 }
+
+
+ESlateVisibility UUIFrameworkWidget::GetVisibility() const
+{
+	return Visibility;
+}
+
+void UUIFrameworkWidget::SetVisibility(ESlateVisibility InVisibility)
+{
+	if (Visibility != InVisibility)
+	{
+		Visibility = InVisibility;
+		MARK_PROPERTY_DIRTY_FROM_NAME(ThisClass, Visibility, this);
+	}
+}
+
+bool UUIFrameworkWidget::IsEnabled() const
+{
+	return bIsEnabled;
+}
+
+void UUIFrameworkWidget::SetEnabled(bool bInIsEnabled)
+{
+	if (bIsEnabled != bInIsEnabled)
+	{
+		bIsEnabled = bInIsEnabled;
+		MARK_PROPERTY_DIRTY_FROM_NAME(ThisClass, bIsEnabled, this);
+	}
+}
+
+void UUIFrameworkWidget::OnRep_IsEnabled()
+{
+	if (LocalUMGWidget)
+	{
+		LocalUMGWidget->SetIsEnabled(bIsEnabled);
+	}
+}
+
+void UUIFrameworkWidget::OnRep_Visibility()
+{
+	if (LocalUMGWidget)
+	{
+		LocalUMGWidget->SetVisibility(Visibility);
+	}
+}
diff --git a/Engine/Plugins/Experimental/UIFramework/Source/Public/UIFWidget.h b/Engine/Plugins/Experimental/UIFramework/Source/Public/UIFWidget.h
--- a/Engine/Plugins/Experimental/UIFramework/Source/Public/UIFWidget.h
+++ b/Engine/Plugins/Experimental/UIFramework/Source/Public/UIFWidget.h
@@ -2,6 +2,7 @@
 
 #pragma once
 
+#include "Components/SlateWrapperTypes.h"
 #include "Components/Widget.h"
 #include "Types/UIFParentWidget.h"
 #include "Types/UIFWidgetId.h"
@@ -20,6 +21,13 @@ class UIFRAMEWORK_API UUIFrameworkWidget : public UObject
 	GENERATED_BODY()
 	friend FUIFrameworkModule;
 
+private:
+	UPROPERTY(BlueprintReadWrite, ReplicatedUsing = "OnRep_IsEnabled", Getter = "IsEnabled", Setter="SetEnabled", Category = "UI Framework", meta = (AllowPrivateAccess = "true"))
+	bool bIsEnabled = true;
+
+	UPROPERTY(BlueprintReadWrite, ReplicatedUsing = "OnRep_Visibility", Getter, Setter = "SetVisibility", Category = "UI Framework", meta = (AllowPrivateAccess = "true"))
+	ESlateVisibility Visibility = ESlateVisibility::Visible;
+
 public:
 	//~ Begin UObject
 	virtual bool IsSupportedForNetworking() const override
@@ -63,10 +71,23 @@ public:
 	virtual void LocalAddChild(FUIFrameworkWidgetId ChildId);
 	void LocalDestroyUMGWidget();
 
+	//~ Properties
+	ESlateVisibility GetVisibility() const;
+	void SetVisibility(ESlateVisibility InVisibility);
+
+	bool IsEnabled() const;
+	void SetEnabled(bool bEnabled);
+
 protected:
 	virtual void AuthorityRemoveChild(UUIFrameworkWidget* Widget) {}
 	virtual void LocalOnUMGWidgetCreated() { }
 
+private:
+	UFUNCTION()
+	void OnRep_IsEnabled();
+	UFUNCTION()
+	void OnRep_Visibility();
+
 protected:
 	UPROPERTY(BlueprintReadOnly, EditDefaultsOnly, Category = "UI Framework")
 	TSoftClassPtr<UWidget> WidgetClass; // todo: make this private and use a constructor argument
@@ -87,10 +108,4 @@ private:
 	//~ LocalOnly
 	UPROPERTY(Transient)
 	TObjectPtr<UWidget> LocalUMGWidget;
-
-	//UPROPERTY(BlueprintReadWrite)
-	//EUIEnability Enablility;
-	//
-	//UPROPERTY(BlueprintReadWrite)
-	//EUIVisibility Visibility;
 };
