UMG: Add ReplaceStackBoxChildAt, fixes the bad behavior of ReplaceChildAt on the WidgetPanel.

diff --git a/Engine/Plugins/Experimental/UIFramework/Source/Private/Widgets/UIFStackBox.cpp b/Engine/Plugins/Experimental/UIFramework/Source/Private/Widgets/UIFStackBox.cpp
--- a/Engine/Plugins/Experimental/UIFramework/Source/Private/Widgets/UIFStackBox.cpp
+++ b/Engine/Plugins/Experimental/UIFramework/Source/Private/Widgets/UIFStackBox.cpp
@@ -6,6 +6,7 @@
 #include "UIFModule.h"
 #include "UIFPlayerComponent.h"
 
+#include "Components/Spacer.h"
 #include "Components/StackBox.h"
 #include "Components/StackBoxSlot.h"
 
@@ -187,8 +188,32 @@ void UUIFrameworkStackBox::LocalAddChild(FUIFrameworkWidgetId ChildId)
 			UWidget* ChildUMGWidget = ChildWidget->LocalGetUMGWidget();
 			if (ensure(ChildUMGWidget))
 			{
-				UPanelSlot* PanelSlot = CastChecked<UStackBox>(LocalGetUMGWidget())->AddChild(ChildUMGWidget);
-				checkf(PanelSlot, TEXT("StackBoxPanel should be able to receive slot"));
+				UStackBox* StackBox = CastChecked<UStackBox>(LocalGetUMGWidget());
+				UPanelSlot* PanelSlot = nullptr;
+				{
+					const int32 ChildCount = StackBox->GetChildrenCount();
+					if (StackBoxEntry->Index < ChildCount)
+					{
+						// replace the existing slot
+						StackBox->ReplaceStackBoxChildAt(StackBoxEntry->Index, ChildUMGWidget);
+						PanelSlot = StackBox->GetSlots()[StackBoxEntry->Index];
+					}
+					else if (StackBoxEntry->Index == ChildCount)
+					{
+						// add to the end
+						PanelSlot = StackBox->AddChild(ChildUMGWidget);
+					}
+					else
+					{
+						// insert until the widget count is correct
+						for (int32 Index = ChildCount; Index < StackBoxEntry->Index; ++Index)
+						{
+							StackBox->AddChild(NewObject<USpacer>(GetTransientPackage()));
+						}
+						PanelSlot = StackBox->AddChild(ChildUMGWidget);
+					}
+					checkf(PanelSlot, TEXT("StackBoxPanel should be able to receive slot"));
+				}
 
 				StackBoxEntry->LocalAquireWidget();
 
diff --git a/Engine/Plugins/Experimental/UIFramework/Source/Public/Widgets/UIFStackBox.h b/Engine/Plugins/Experimental/UIFramework/Source/Public/Widgets/UIFStackBox.h
--- a/Engine/Plugins/Experimental/UIFramework/Source/Public/Widgets/UIFStackBox.h
+++ b/Engine/Plugins/Experimental/UIFramework/Source/Public/Widgets/UIFStackBox.h
@@ -12,6 +12,7 @@
 #include "UIFStackBox.generated.h"
 
 class UUIFrameworkStackBox;
+struct FUIFrameworkStackBoxSlotList;
 
 /**
  *
@@ -21,6 +22,7 @@ struct FUIFrameworkStackBoxSlot : public FUIFrameworkSlotBase
 {
 	GENERATED_BODY()
 
+	friend UUIFrameworkStackBox;
 	friend FUIFrameworkStackBoxSlotList;
 
 	/** Horizontal alignment of the widget inside the slot. */
diff --git a/Engine/Source/Runtime/UMG/Private/Components/StackBox.cpp b/Engine/Source/Runtime/UMG/Private/Components/StackBox.cpp
--- a/Engine/Source/Runtime/UMG/Private/Components/StackBox.cpp
+++ b/Engine/Source/Runtime/UMG/Private/Components/StackBox.cpp
@@ -75,6 +75,18 @@ UStackBoxSlot* UStackBox::AddChildToStackBox(UWidget* Content)
 	return Cast<UStackBoxSlot>(Super::AddChild(Content));
 }
 
+bool UStackBox::ReplaceStackBoxChildAt(int32 Index, UWidget* Content)
+{
+	if (!Slots.IsValidIndex(Index) || Content == nullptr)
+	{
+		return false;
+	}
+
+	UStackBoxSlot* StackBoxSlot = CastChecked<UStackBoxSlot>(Slots[Index]);
+	StackBoxSlot->ReplaceContent(Content);
+	return true;
+}
+
 TSharedRef<SWidget> UStackBox::RebuildWidget()
 {
 	MyBox = SNew(SStackBox)
diff --git a/Engine/Source/Runtime/UMG/Private/Components/StackBoxSlot.cpp b/Engine/Source/Runtime/UMG/Private/Components/StackBoxSlot.cpp
--- a/Engine/Source/Runtime/UMG/Private/Components/StackBoxSlot.cpp
+++ b/Engine/Source/Runtime/UMG/Private/Components/StackBoxSlot.cpp
@@ -32,6 +32,28 @@ void UStackBoxSlot::BuildSlot(TSharedRef<SStackBox> StackBox)
 		];
 }
 
+void UStackBoxSlot::ReplaceContent(UWidget* NewContent)
+{
+	if (NewContent)
+	{
+		NewContent->RemoveFromParent();
+
+		{
+			UWidget* PreviousWidget = Content;
+			PreviousWidget->Slot = nullptr;
+			PreviousWidget->RemoveFromParent();
+		}
+
+		Content = NewContent;
+		NewContent->Slot = this;
+
+		if (Slot)
+		{
+			Slot->AttachWidget(Content->TakeWidget());
+		}
+	}
+}
+
 FMargin UStackBoxSlot::GetPadding() const
 {
 	return Padding;
diff --git a/Engine/Source/Runtime/UMG/Public/Components/StackBox.h b/Engine/Source/Runtime/UMG/Public/Components/StackBox.h
--- a/Engine/Source/Runtime/UMG/Public/Components/StackBox.h
+++ b/Engine/Source/Runtime/UMG/Public/Components/StackBox.h
@@ -39,6 +39,10 @@ public:
 	UFUNCTION(BlueprintCallable, Category="Panel")
 	UStackBoxSlot* AddChildToStackBox(UWidget* Content);
 
+	/** Replace the widget at the given index it with a different widget. */
+	UFUNCTION(BlueprintCallable, Category = "Panel")
+	bool ReplaceStackBoxChildAt(int32 Index, UWidget* Content);
+
 #if WITH_EDITOR
 	//~ UWidget interface
 	virtual const FText GetPaletteCategory() override;
diff --git a/Engine/Source/Runtime/UMG/Public/Components/StackBoxSlot.h b/Engine/Source/Runtime/UMG/Public/Components/StackBoxSlot.h
--- a/Engine/Source/Runtime/UMG/Public/Components/StackBoxSlot.h
+++ b/Engine/Source/Runtime/UMG/Public/Components/StackBoxSlot.h
@@ -62,6 +62,9 @@ public:
 	/** Builds the underlying FSlot for the Slate layout panel. */
 	void BuildSlot(TSharedRef<SStackBox> InBox);
 
+	/** Replace the slot content. */
+	void ReplaceContent(UWidget* Content);
+
 #if WITH_EDITOR
 	virtual bool NudgeByDesigner(const FVector2D& NudgeDirection, const TOptional<int32>& GridSnapSize) override;
 	virtual void SynchronizeFromTemplate(const UPanelSlot* const TemplateSlot) override;
