First version of the input option to VerseUI.
UIFramework: Remove All widgets from actor channel.

diff --git a/Engine/Plugins/Experimental/UIFramework/Source/Private/Types/UIFWidgetTree.cpp b/Engine/Plugins/Experimental/UIFramework/Source/Private/Types/UIFWidgetTree.cpp
--- a/Engine/Plugins/Experimental/UIFramework/Source/Private/Types/UIFWidgetTree.cpp
+++ b/Engine/Plugins/Experimental/UIFramework/Source/Private/Types/UIFWidgetTree.cpp
@@ -259,7 +259,34 @@ bool FUIFrameworkWidgetTree::AuthorityRemoveChildRecursiveInternal(UUIFrameworkW
 		return true;
 	}
 	return false;
+}
 
+void FUIFrameworkWidgetTree::AuthorityAddAllWidgetsFromActorChannel()
+{
+	if (ensure(ReplicatedOwner) && ReplicatedOwner->IsUsingRegisteredSubObjectList())
+	{
+		for (const FUIFrameworkWidgetTreeEntry& Entry : Entries)
+		{
+			if (Entry.Child)
+			{
+				ReplicatedOwner->AddReplicatedSubObject(Entry.Child);
+			}
+		}
+	}
+}
+
+void FUIFrameworkWidgetTree::AuthorityRemoveAllWidgetsFromActorChannel()
+{
+	if (ensure(ReplicatedOwner) && ReplicatedOwner->IsUsingRegisteredSubObjectList())
+	{
+		for (const FUIFrameworkWidgetTreeEntry& Entry : Entries)
+		{
+			if (Entry.Child && ReplicatedOwner->IsReplicatedSubObjectRegistered(Entry.Child))
+			{
+				ReplicatedOwner->RemoveReplicatedSubObject(Entry.Child);
+			}
+		}
+	}
 }
 
 FUIFrameworkWidgetTreeEntry* FUIFrameworkWidgetTree::LocalGetEntryByReplicationId(int32 ReplicationId)
diff --git a/Engine/Plugins/Experimental/UIFramework/Source/Private/UIFModule.cpp b/Engine/Plugins/Experimental/UIFramework/Source/Private/UIFModule.cpp
--- a/Engine/Plugins/Experimental/UIFramework/Source/Private/UIFModule.cpp
+++ b/Engine/Plugins/Experimental/UIFramework/Source/Private/UIFModule.cpp
@@ -4,9 +4,15 @@
 
 #include "Modules/ModuleManager.h"
 #include "UIFPlayerComponent.h"
+#include "UIFPresenter.h"
 #include "UIFWidget.h"
 #include "UObject/Package.h"
 
+namespace UE::UIFramework::Private
+{
+	static TSubclassOf<UUIFrameworkPresenter> Director;
+}
+
 UUIFrameworkWidget* FUIFrameworkModule::AuthorityAttachWidget(UUIFrameworkPlayerComponent* ReplicationOwner, FUIFrameworkParentWidget Parent, UUIFrameworkWidget* Child)
 {
 	check(Child);
@@ -93,4 +99,14 @@ void FUIFrameworkModule::AuthorityDetachWidgetFromParent(UUIFrameworkWidget* Chi
 	}
 }
 
+void FUIFrameworkModule::SetPresenterClass(TSubclassOf<UUIFrameworkPresenter> InDirector)
+{
+	UE::UIFramework::Private::Director = InDirector;
+}
+
+TSubclassOf<UUIFrameworkPresenter> FUIFrameworkModule::GetPresenterClass()
+{
+	return UE::UIFramework::Private::Director.Get() ? UE::UIFramework::Private::Director : TSubclassOf<UUIFrameworkPresenter>(UUIFrameworkGameViewportPresenter::StaticClass());
+}
+
 IMPLEMENT_MODULE(FUIFrameworkModule, UIFramework)
\ No newline at end of file
diff --git a/Engine/Plugins/Experimental/UIFramework/Source/Private/UIFPlayerComponent.cpp b/Engine/Plugins/Experimental/UIFramework/Source/Private/UIFPlayerComponent.cpp
--- a/Engine/Plugins/Experimental/UIFramework/Source/Private/UIFPlayerComponent.cpp
+++ b/Engine/Plugins/Experimental/UIFramework/Source/Private/UIFPlayerComponent.cpp
@@ -3,8 +3,8 @@
 #include "UIFPlayerComponent.h"
 #include "UIFLog.h"
 #include "UIFModule.h"
+#include "UIFPresenter.h"
 #include "UIFWidget.h"
-#include "Blueprint/GameViewportSubsystem.h"
 
 #include "Blueprint/UserWidget.h"
 #include "Engine/ActorChannel.h"
@@ -70,12 +70,31 @@ UUIFrameworkPlayerComponent::UUIFrameworkPlayerComponent()
 	PrimaryComponentTick.bCanEverTick = true;
 	PrimaryComponentTick.bStartWithTickEnabled = false;
 	PrimaryComponentTick.TickGroup = ETickingGroup::TG_DuringPhysics;
+
+	if (!IsTemplate() && GetPlayerController() && GetPlayerController()->IsLocalPlayerController())
+	{
+		Presenter = static_cast<UUIFrameworkPresenter*>(CreateDefaultSubobject(TEXT("Presenter"), UUIFrameworkPresenter::StaticClass(), FUIFrameworkModule::GetPresenterClass().Get(), true, true));
+	}
+}
+
+void UUIFrameworkPlayerComponent::InitializeComponent()
+{
+	Super::InitializeComponent();
+
+	if (GetOwner()->HasAuthority())
+	{
+		WidgetTree.AuthorityAddAllWidgetsFromActorChannel();
+	}
 }
 
 void UUIFrameworkPlayerComponent::UninitializeComponent()
 {
 	// On local, remove all UWidget.
 	if (!GetOwner()->HasAuthority())
+	{
+		WidgetTree.AuthorityRemoveAllWidgetsFromActorChannel();
+	}
+	else
 	{
 		for (FUIFrameworkGameLayerSlot& Entry : RootList.Entries)
 		{
@@ -85,6 +104,8 @@ void UUIFrameworkPlayerComponent::UninitializeComponent()
 			}
 		}
 	}
+
+	Super::UninitializeComponent();
 }
 
 void UUIFrameworkPlayerComponent::GetLifetimeReplicatedProps(TArray< FLifetimeProperty >& OutLifetimeProps) const
@@ -343,32 +364,14 @@ void UUIFrameworkPlayerComponent::LocalAddChild(FUIFrameworkWidgetId WidgetId)
 			if (ensure(UMGWidget))
 			{
 				UMGWidget->RemoveFromParent();
-				if (UGameViewportSubsystem* Subsystem = UGameViewportSubsystem::Get(GetWorld()))
-				{
-					LayerEntry->LocalAquireWidget();
-
-					FGameViewportWidgetSlot GameViewportWidgetSlot;
-					GameViewportWidgetSlot.ZOrder = LayerEntry->ZOrder;
-					if (LayerEntry->Type == EUIFrameworkGameLayerType::Viewport)
-					{
-						Subsystem->AddWidget(UMGWidget, GameViewportWidgetSlot);
-					}
-					else
-					{
-						APlayerController* LocalOwner = GetPlayerController();
-						check(LocalOwner);
-						Subsystem->AddWidgetForPlayer(UMGWidget, LocalOwner->GetLocalPlayer(), GameViewportWidgetSlot);
-					}
-				}
+				LayerEntry->LocalAquireWidget();
+				check(Presenter);
+				Presenter->AddToViewport(UMGWidget, *LayerEntry);
 			}
 		}
 		else
 		{
 			UE_LOG(LogUIFramework, Log, TEXT("The widget '%" INT64_FMT "' doesn't exist in the WidgetTree."), WidgetId.GetKey());
-			if (Widget && Widget->LocalGetUMGWidget())
-			{
-				Widget->LocalGetUMGWidget()->RemoveFromParent();
-			}
 		}
 	}
 	else
diff --git a/Engine/Plugins/Experimental/UIFramework/Source/Private/UIFPresenter.cpp b/Engine/Plugins/Experimental/UIFramework/Source/Private/UIFPresenter.cpp
new file mode 100644
--- /dev/null
+++ b/Engine/Plugins/Experimental/UIFramework/Source/Private/UIFPresenter.cpp
@@ -0,0 +1,27 @@
+// Copyright Epic Games, Inc. All Rights Reserved.
+
+#include "UIFPresenter.h"
+#include "UIFPlayerComponent.h"
+
+#include "Blueprint/GameViewportSubsystem.h"
+#include "Components/Widget.h"
+#include "GameFramework/PlayerController.h"
+
+void UUIFrameworkGameViewportPresenter::AddToViewport(UWidget* UMGWidget, const FUIFrameworkGameLayerSlot& Slot)
+{
+	if (UGameViewportSubsystem* Subsystem = UGameViewportSubsystem::Get(GetOuterUUIFrameworkPlayerComponent()->GetWorld()))
+	{
+		FGameViewportWidgetSlot GameViewportWidgetSlot;
+		GameViewportWidgetSlot.ZOrder = Slot.ZOrder;
+		if (Slot.Type == EUIFrameworkGameLayerType::Viewport)
+		{
+			Subsystem->AddWidget(UMGWidget, GameViewportWidgetSlot);
+		}
+		else
+		{
+			APlayerController* LocalOwner = GetOuterUUIFrameworkPlayerComponent()->GetPlayerController();
+			check(LocalOwner);
+			Subsystem->AddWidgetForPlayer(UMGWidget, LocalOwner->GetLocalPlayer(), GameViewportWidgetSlot);
+		}
+	}
+}
diff --git a/Engine/Plugins/Experimental/UIFramework/Source/Public/Types/UIFWidgetTree.h b/Engine/Plugins/Experimental/UIFramework/Source/Public/Types/UIFWidgetTree.h
--- a/Engine/Plugins/Experimental/UIFramework/Source/Public/Types/UIFWidgetTree.h
+++ b/Engine/Plugins/Experimental/UIFramework/Source/Public/Types/UIFWidgetTree.h
@@ -107,6 +107,12 @@ public:
 	UUIFrameworkWidget* FindWidgetById(FUIFrameworkWidgetId WidgetId);
 	const UUIFrameworkWidget* FindWidgetById(FUIFrameworkWidgetId WidgetId) const;
 
+	/** Add all widgets in the tree to the ActorChannel replicated list */
+	void AuthorityAddAllWidgetsFromActorChannel();
+
+	/** Removes all widgets added to the ActorChannel replicated list */
+	void AuthorityRemoveAllWidgetsFromActorChannel();
+
 #if UE_UIFRAMEWORK_WITH_DEBUG
 	void AuthorityTest() const;
 #endif
diff --git a/Engine/Plugins/Experimental/UIFramework/Source/Public/UIFModule.h b/Engine/Plugins/Experimental/UIFramework/Source/Public/UIFModule.h
--- a/Engine/Plugins/Experimental/UIFramework/Source/Public/UIFModule.h
+++ b/Engine/Plugins/Experimental/UIFramework/Source/Public/UIFModule.h
@@ -6,8 +6,9 @@
 #include "Modules/ModuleInterface.h"
 #include "Types/UIFParentWidget.h"
 
-class UUIFrameworkWidget;
 class UUIFrameworkPlayerComponent;
+class UUIFrameworkPresenter;
+class UUIFrameworkWidget;
 
 /**
  * 
@@ -27,10 +28,10 @@ public:
 	 * Will remove the widget from the tree and the replication owner.
 	 */
 	static void AuthorityDetachWidgetFromParent(UUIFrameworkWidget* Child);
-	
-	using FOldWidgetToNewWidgetMap = TMap<UUIFrameworkWidget*, UUIFrameworkWidget*>;
-	DECLARE_DELEGATE_OneParam(FOnReattachWidgets, FOldWidgetToNewWidgetMap);
-	FOnReattachWidgets AuthorityOnReattachWidgets;
+
+	//~ this should be a project setting
+	static void SetPresenterClass(TSubclassOf<UUIFrameworkPresenter> Director);
+	static TSubclassOf<UUIFrameworkPresenter> GetPresenterClass();
 
 private:
 	static void AuthoritySetParentReplicationOwnerRecursive(UUIFrameworkWidget* Widget);
diff --git a/Engine/Plugins/Experimental/UIFramework/Source/Public/UIFPlayerComponent.h b/Engine/Plugins/Experimental/UIFramework/Source/Public/UIFPlayerComponent.h
--- a/Engine/Plugins/Experimental/UIFramework/Source/Public/UIFPlayerComponent.h
+++ b/Engine/Plugins/Experimental/UIFramework/Source/Public/UIFPlayerComponent.h
@@ -11,6 +11,7 @@
 #include "UIFPlayerComponent.generated.h"
 
 class UUIFrameworkPlayerComponent;
+class UUIFrameworkPresenter;
 class UUIFrameworkWidget;
 class UWidget;
 struct FStreamableHandle;
@@ -26,6 +27,18 @@ enum class EUIFrameworkGameLayerType : uint8
 	PlayerScreen,
 };
 
+/**
+ *
+ */
+UENUM(BlueprintType)
+enum class EUIFrameworkInputMode : uint8
+{
+	// Input is received by the UI.
+	UI,
+	// Input is received by the Game.
+	Game,
+};
+
 
 /**
  *
@@ -39,6 +52,9 @@ struct FUIFrameworkGameLayerSlot : public FUIFrameworkSlotBase
 
 	UPROPERTY(BlueprintReadWrite, Category = "UI Framework")
 	int32 ZOrder = 0;
+	
+	UPROPERTY(BlueprintReadWrite, Category = "UI Framework")
+	EUIFrameworkInputMode InputMode = EUIFrameworkInputMode::Game;
 
 	UPROPERTY(BlueprintReadWrite, Category = "UI Framework")
 	EUIFrameworkGameLayerType Type = EUIFrameworkGameLayerType::Viewport;
@@ -129,6 +145,7 @@ public:
 	}
 
 	//~ Begin UActorComponent
+	virtual void InitializeComponent() override;
 	virtual void UninitializeComponent() override;
 	virtual bool ReplicateSubobjects(class UActorChannel* Channel, class FOutBunch* Bunch, FReplicationFlags* RepFlags) override;
 	virtual void TickComponent(float DeltaTime, enum ELevelTick TickType, FActorComponentTickFunction* ThisTickFunction) override;
@@ -150,6 +167,9 @@ private:
 	UPROPERTY(Replicated)
 	FUIFrameworkWidgetTree WidgetTree;
 
+	UPROPERTY(Transient)
+	TObjectPtr<UUIFrameworkPresenter> Presenter;
+
 	//~ Widget can be net replicated but not constructed yet.
 	UPROPERTY(Transient)
 	TSet<int32> NetReplicationPending;
diff --git a/Engine/Plugins/Experimental/UIFramework/Source/Public/UIFPresenter.h b/Engine/Plugins/Experimental/UIFramework/Source/Public/UIFPresenter.h
new file mode 100644
--- /dev/null
+++ b/Engine/Plugins/Experimental/UIFramework/Source/Public/UIFPresenter.h
@@ -0,0 +1,38 @@
+// Copyright Epic Games, Inc. All Rights Reserved.
+
+#pragma once
+
+#include "UObject/Object.h"
+#include "UIFPlayerComponent.h"
+
+#include "UIFPresenter.generated.h"
+
+class FUIFrameworkModule;
+
+/**
+ * 
+ */
+ UCLASS(Abstract, Within=UIFrameworkPlayerComponent)
+class UIFRAMEWORK_API UUIFrameworkPresenter : public UObject
+{
+	GENERATED_BODY()
+
+public:
+	virtual void AddToViewport(UWidget* UMGWidget, const FUIFrameworkGameLayerSlot& Slot)
+	{
+
+	}
+};
+
+
+/**
+ *
+ */
+ UCLASS()
+class UIFRAMEWORK_API UUIFrameworkGameViewportPresenter : public UUIFrameworkPresenter
+ {
+	 GENERATED_BODY()
+
+public:
+	virtual void AddToViewport(UWidget* UMGWidget, const FUIFrameworkGameLayerSlot& Slot) override;
+};
