UIFramework: Fix bad replication test. Fix existing widget readded to a new widget.

diff --git a/Engine/Plugins/Experimental/UIFramework/Source/Private/UIFPlayerComponent.cpp b/Engine/Plugins/Experimental/UIFramework/Source/Private/UIFPlayerComponent.cpp
--- a/Engine/Plugins/Experimental/UIFramework/Source/Private/UIFPlayerComponent.cpp
+++ b/Engine/Plugins/Experimental/UIFramework/Source/Private/UIFPlayerComponent.cpp
@@ -90,7 +90,7 @@ void UUIFrameworkPlayerComponent::InitializeComponent()
 void UUIFrameworkPlayerComponent::UninitializeComponent()
 {
 	// On local, remove all UWidget.
-	if (!GetOwner()->HasAuthority())
+	if (GetOwner()->HasAuthority())
 	{
 		WidgetTree.AuthorityRemoveAllWidgetsFromActorChannel();
 	}
diff --git a/Engine/Plugins/Experimental/UIFramework/Source/Private/Widgets/UIFButton.cpp b/Engine/Plugins/Experimental/UIFramework/Source/Private/Widgets/UIFButton.cpp
--- a/Engine/Plugins/Experimental/UIFramework/Source/Private/Widgets/UIFButton.cpp
+++ b/Engine/Plugins/Experimental/UIFramework/Source/Private/Widgets/UIFButton.cpp
@@ -48,7 +48,7 @@ void UUIFrameworkButton::SetContent(FUIFrameworkSimpleSlot InEntry)
 		if (InEntry.AuthorityGetWidget())
 		{
 			UUIFrameworkPlayerComponent* PreviousOwner = InEntry.AuthorityGetWidget()->GetPlayerComponent();
-			if (PreviousOwner != nullptr && PreviousOwner != GetPlayerComponent())
+			if (PreviousOwner != nullptr && GetPlayerComponent() != nullptr && PreviousOwner != GetPlayerComponent())
 			{
 				Slot.AuthoritySetWidget(nullptr);
 				FFrame::KismetExecutionMessage(TEXT("The widget was created for another player. It can't be added."), ELogVerbosity::Warning, "InvalidPlayerParent");
diff --git a/Engine/Plugins/Experimental/UIFramework/Source/Private/Widgets/UIFCanvasBox.cpp b/Engine/Plugins/Experimental/UIFramework/Source/Private/Widgets/UIFCanvasBox.cpp
--- a/Engine/Plugins/Experimental/UIFramework/Source/Private/Widgets/UIFCanvasBox.cpp
+++ b/Engine/Plugins/Experimental/UIFramework/Source/Private/Widgets/UIFCanvasBox.cpp
@@ -99,7 +99,7 @@ void UUIFrameworkCanvasBox::AddWidget(FUIFrameworkCanvasBoxSlot InEntry)
 	{
 		FFrame::KismetExecutionMessage(TEXT("The widget is invalid. It can't be added."), ELogVerbosity::Warning, "InvalidWidgetToAdd");
 	}
-	else if (GetPlayerComponent() && GetPlayerComponent() != InEntry.AuthorityGetWidget()->GetPlayerComponent())
+	else if (GetPlayerComponent() && InEntry.AuthorityGetWidget()->GetPlayerComponent() && GetPlayerComponent() != InEntry.AuthorityGetWidget()->GetPlayerComponent())
 	{
 		check(GetPlayerComponent()->GetOwner()->HasAuthority());
 		FFrame::KismetExecutionMessage(TEXT("The widget was created for another player. It can't be removed on this player."), ELogVerbosity::Warning, "InvalidPlayerParentOnRemovedWidget");
diff --git a/Engine/Plugins/Experimental/UIFramework/Source/Private/Widgets/UIFStackBox.cpp b/Engine/Plugins/Experimental/UIFramework/Source/Private/Widgets/UIFStackBox.cpp
--- a/Engine/Plugins/Experimental/UIFramework/Source/Private/Widgets/UIFStackBox.cpp
+++ b/Engine/Plugins/Experimental/UIFramework/Source/Private/Widgets/UIFStackBox.cpp
@@ -105,7 +105,7 @@ void UUIFrameworkStackBox::AddWidget(FUIFrameworkStackBoxSlot InEntry)
 	{
 		FFrame::KismetExecutionMessage(TEXT("The widget is invalid. It can't be added."), ELogVerbosity::Warning, "InvalidWidgetToAdd");
 	}
-	else if (GetPlayerComponent() && GetPlayerComponent() != InEntry.AuthorityGetWidget()->GetPlayerComponent())
+	else if (GetPlayerComponent() && InEntry.AuthorityGetWidget()->GetPlayerComponent() && GetPlayerComponent() != InEntry.AuthorityGetWidget()->GetPlayerComponent())
 	{
 		check(GetPlayerComponent()->GetOwner()->HasAuthority());
 		FFrame::KismetExecutionMessage(TEXT("The widget was created for another player. It can't be removed on this player."), ELogVerbosity::Warning, "InvalidPlayerParentOnRemovedWidget");
