Iris
- Added missing dependencies where needed.

diff --git a/Engine/Plugins/Experimental/UIFramework/Source/UIFramework.Build.cs b/Engine/Plugins/Experimental/UIFramework/Source/UIFramework.Build.cs
--- a/Engine/Plugins/Experimental/UIFramework/Source/UIFramework.Build.cs
+++ b/Engine/Plugins/Experimental/UIFramework/Source/UIFramework.Build.cs
@@ -29,5 +29,7 @@ public class UIFramework : ModuleRules
 				"NetCore",
 			}
 		);
+
+		SetupIrisSupport(Target);
 	}
 }
