UIFramework: Add callback for authority to know if the client removed a widget that was not authoritive. Hold a reference to the the owner widget wrapper (to keep it alive for button click and event). It will also be use when we are going to rename the object.

diff --git a/Engine/Plugins/Experimental/UIFramework/Source/Private/Types/UIFWidgetTree.cpp b/Engine/Plugins/Experimental/UIFramework/Source/Private/Types/UIFWidgetTree.cpp
--- a/Engine/Plugins/Experimental/UIFramework/Source/Private/Types/UIFWidgetTree.cpp
+++ b/Engine/Plugins/Experimental/UIFramework/Source/Private/Types/UIFWidgetTree.cpp
@@ -61,6 +61,7 @@ FUIFrameworkWidgetTree::FUIFrameworkWidgetTree(AActor* InReplicatedOwner, IUIFra
 	: ReplicatedOwner(InReplicatedOwner)
 	, Owner(InOwner)
 {
+	check(Owner);
 #if UE_UIFRAMEWORK_WITH_DEBUG
 	UE::UIFramework::Private::GTrees.Add(this);
 #endif
@@ -132,6 +133,9 @@ void FUIFrameworkWidgetTree::PostReplicatedChange(const TArrayView<int32>& Chang
 	}
 }
 
+void FUIFrameworkWidgetTree::PostReplicatedReceive(FPostReplicatedReceiveParameters Param)
+{}
+
 bool FUIFrameworkWidgetTree::ReplicateSubWidgets(UActorChannel* Channel, FOutBunch* Bunch, FReplicationFlags* RepFlags)
 {
 	bool bWroteSomething = false;
@@ -158,17 +162,17 @@ bool FUIFrameworkWidgetTree::ReplicateSubWidgets(UActorChannel* Channel, FOutBun
 void FUIFrameworkWidgetTree::AuthorityAddRoot(UUIFrameworkWidget* Widget)
 {
 	check(Widget);
-	AuthorityAddChildInternal(nullptr, Widget);
+	AuthorityAddChildInternal(nullptr, Widget, true);
 }
 
 void FUIFrameworkWidgetTree::AuthorityAddWidget(UUIFrameworkWidget* Parent, UUIFrameworkWidget* Child)
 {
 	check(Parent);
 	check(Child);
-	AuthorityAddChildInternal(Parent, Child);
+	AuthorityAddChildInternal(Parent, Child, true);
 }
 
-void FUIFrameworkWidgetTree::AuthorityAddChildInternal(UUIFrameworkWidget* Parent, UUIFrameworkWidget* Child)
+void FUIFrameworkWidgetTree::AuthorityAddChildInternal(UUIFrameworkWidget* Parent, UUIFrameworkWidget* Child, bool bFirst)
 {
 	if (int32* PreviousEntryIndexPtr = AuthorityIndexByWidgetMap.Find(Child))
 	{
@@ -197,6 +201,10 @@ void FUIFrameworkWidgetTree::AuthorityAddChildInternal(UUIFrameworkWidget* Paren
 		}
 
 		AuthorityAddChildRecursiveInternal(Child);
+		if (bFirst)
+		{
+			AuthorityOnWidgetAdded.Broadcast(Child);
+		}
 	}
 }
 
@@ -207,7 +215,7 @@ void FUIFrameworkWidgetTree::AuthorityAddChildRecursiveInternal(UUIFrameworkWidg
 		{
 			if (ChildWidget != nullptr)
 			{
-				Self->AuthorityAddChildInternal(InParentWidget, ChildWidget);
+				Self->AuthorityAddChildInternal(InParentWidget, ChildWidget, false);
 			}
 		});
 }
@@ -219,6 +227,17 @@ void FUIFrameworkWidgetTree::AuthorityRemoveWidget(UUIFrameworkWidget* Widget)
 	if (AuthorityRemoveChildRecursiveInternal(Widget))
 	{
 		MarkArrayDirty();
+		AuthorityOnWidgetRemoved.Broadcast(Widget);
+	}
+}
+
+void FUIFrameworkWidgetTree::LocalRemoveRoot(const UUIFrameworkWidget* Widget)
+{
+	check(Widget);
+
+	if (ensure(Owner))
+	{
+		Owner->LocalRemoveWidgetRootFromTree(Widget);
 	}
 }
 
diff --git a/Engine/Plugins/Experimental/UIFramework/Source/Private/UIFPlayerComponent.cpp b/Engine/Plugins/Experimental/UIFramework/Source/Private/UIFPlayerComponent.cpp
--- a/Engine/Plugins/Experimental/UIFramework/Source/Private/UIFPlayerComponent.cpp
+++ b/Engine/Plugins/Experimental/UIFramework/Source/Private/UIFPlayerComponent.cpp
@@ -56,6 +56,11 @@ FUIFrameworkGameLayerSlot* FUIFrameworkGameLayerSlotList::FindEntry(FUIFramework
 	return Entries.FindByPredicate([WidgetId](const FUIFrameworkGameLayerSlot& Entry) { return Entry.GetWidgetId() == WidgetId; });
 }
 
+const FUIFrameworkGameLayerSlot* FUIFrameworkGameLayerSlotList::FindEntry(FUIFrameworkWidgetId WidgetId) const
+{
+	return Entries.FindByPredicate([WidgetId](const FUIFrameworkGameLayerSlot& Entry) { return Entry.GetWidgetId() == WidgetId; });
+}
+
 
 /**
  *
@@ -312,6 +317,20 @@ void UUIFrameworkPlayerComponent::LocalWidgetRemovedFromTree(const FUIFrameworkW
 	PrimaryComponentTick.SetTickFunctionEnable(NetReplicationPending.Num() > 0 || AddPending.Num() > 0 || ClassesToLoad.Num() > 0);
 }
 
+void UUIFrameworkPlayerComponent::LocalRemoveWidgetRootFromTree(const UUIFrameworkWidget* Widget)
+{
+	check(Widget);
+	ServerRemoveWidgetRootFromTree(Widget->GetWidgetId());
+}
+
+void UUIFrameworkPlayerComponent::ServerRemoveWidgetRootFromTree_Implementation(FUIFrameworkWidgetId WidgetId)
+{
+	if (UUIFrameworkWidget* Widget = GetWidgetTree().FindWidgetById(WidgetId))
+	{
+		GetWidgetTree().AuthorityRemoveWidget(Widget);
+	}
+}
+
 void UUIFrameworkPlayerComponent::LocalOnClassLoaded(TSoftClassPtr<UWidget> WidgetClass)
 {
 	FWidgetClassToLoad FoundWidgetClassToLoad;
diff --git a/Engine/Plugins/Experimental/UIFramework/Source/Public/Types/UIFWidgetTree.h b/Engine/Plugins/Experimental/UIFramework/Source/Public/Types/UIFWidgetTree.h
--- a/Engine/Plugins/Experimental/UIFramework/Source/Public/Types/UIFWidgetTree.h
+++ b/Engine/Plugins/Experimental/UIFramework/Source/Public/Types/UIFWidgetTree.h
@@ -24,9 +24,19 @@ class IUIFrameworkWidgetTreeOwner
 {
 public:
 	/** A widget was added to the tree. */
-	virtual void LocalWidgetWasAddedToTree(const FUIFrameworkWidgetTreeEntry& Entry) {};
+	virtual void LocalWidgetWasAddedToTree(const FUIFrameworkWidgetTreeEntry& Entry)
+	{
+	}
+
 	/** A widget was removed to the tree. */
-	virtual void LocalWidgetRemovedFromTree(const FUIFrameworkWidgetTreeEntry& Entry) {};
+	virtual void LocalWidgetRemovedFromTree(const FUIFrameworkWidgetTreeEntry& Entry)
+	{
+	}
+
+	/** Remove the widget (and the child) from the server. */
+	virtual void LocalRemoveWidgetRootFromTree(const UUIFrameworkWidget* Widget)
+	{
+	}
 };
 
 /**
@@ -77,6 +87,7 @@ public:
 	void PreReplicatedRemove(const TArrayView<int32> RemovedIndices, int32 FinalSize);
 	void PostReplicatedAdd(const TArrayView<int32> AddedIndices, int32 FinalSize);
 	void PostReplicatedChange(const TArrayView<int32>& ChangedIndices, int32 FinalSize);
+	void PostReplicatedReceive(FPostReplicatedReceiveParameters Param);
 	//~ End of FFastArraySerializer
 
 	bool NetDeltaSerialize(FNetDeltaSerializeInfo& DeltaParms)
@@ -99,6 +110,11 @@ public:
 	 */
 	void AuthorityRemoveWidget(UUIFrameworkWidget* Widget);
 
+	/**
+	 * The widget was removed from the client and the Authority is not aware of it.
+	 */
+	void LocalRemoveRoot(const UUIFrameworkWidget* Widget);
+
 	//~ It is not safe to use the ReplicationId on the Authority because any add or remove would clear the ItemMap.
 	FUIFrameworkWidgetTreeEntry* LocalGetEntryByReplicationId(int32 WidgetId);
 	const FUIFrameworkWidgetTreeEntry* LocalGetEntryByReplicationId(int32 WidgetId) const;
@@ -117,8 +133,12 @@ public:
 	void AuthorityTest() const;
 #endif
 
+	DECLARE_MULTICAST_DELEGATE_OneParam(FUIFrameworkWidgetDelegate, UUIFrameworkWidget*);
+	FUIFrameworkWidgetDelegate AuthorityOnWidgetAdded;
+	FUIFrameworkWidgetDelegate AuthorityOnWidgetRemoved;
+
 private:
-	void AuthorityAddChildInternal(UUIFrameworkWidget* Parent, UUIFrameworkWidget* Child);
+	void AuthorityAddChildInternal(UUIFrameworkWidget* Parent, UUIFrameworkWidget* Child, bool bFirst);
 	void AuthorityAddChildRecursiveInternal(UUIFrameworkWidget* Widget);
 	bool AuthorityRemoveChildRecursiveInternal(UUIFrameworkWidget* Widget);
 
diff --git a/Engine/Plugins/Experimental/UIFramework/Source/Public/UIFPlayerComponent.h b/Engine/Plugins/Experimental/UIFramework/Source/Public/UIFPlayerComponent.h
--- a/Engine/Plugins/Experimental/UIFramework/Source/Public/UIFPlayerComponent.h
+++ b/Engine/Plugins/Experimental/UIFramework/Source/Public/UIFPlayerComponent.h
@@ -89,6 +89,7 @@ public:
 	void AddEntry(FUIFrameworkGameLayerSlot Entry);
 	bool RemoveEntry(UUIFrameworkWidget* Layer);
 	FUIFrameworkGameLayerSlot* FindEntry(FUIFrameworkWidgetId WidgetId);
+	const FUIFrameworkGameLayerSlot* FindEntry(FUIFrameworkWidgetId WidgetId) const;
 
 private:
 	UPROPERTY()
@@ -128,6 +129,11 @@ public:
 	{
 		return WidgetTree;
 	}
+	
+	const FUIFrameworkGameLayerSlotList& GetRootList() const
+	{
+		return RootList;
+	}
 
 	/** Gets the controller that owns the component, this will always be valid during gameplay but can return null in the editor */
 	template <class T = APlayerController>
@@ -155,8 +161,12 @@ public:
 
 	virtual void LocalWidgetWasAddedToTree(const FUIFrameworkWidgetTreeEntry& Entry) override;
 	virtual void LocalWidgetRemovedFromTree(const FUIFrameworkWidgetTreeEntry& Entry) override;
+	virtual void LocalRemoveWidgetRootFromTree(const UUIFrameworkWidget* Widget) override;
 
 private:
+	UFUNCTION(Server, Reliable)
+	void ServerRemoveWidgetRootFromTree(FUIFrameworkWidgetId WidgetId);
+
 	void LocalOnClassLoaded(TSoftClassPtr<UWidget> WidgetClass);
 	void LocalAddChild(FUIFrameworkWidgetId WidgetId);
 
diff --git a/Engine/Plugins/Experimental/UIFramework/Source/Public/UIFWidget.h b/Engine/Plugins/Experimental/UIFramework/Source/Public/UIFWidget.h
--- a/Engine/Plugins/Experimental/UIFramework/Source/Public/UIFWidget.h
+++ b/Engine/Plugins/Experimental/UIFramework/Source/Public/UIFWidget.h
@@ -12,6 +12,20 @@
 
 class FUIFrameworkModule;
 
+/**
+ *
+ */
+UINTERFACE(MinimalAPI)
+class UUIFrameworkWidgetWrapperInterface : public UInterface
+{
+	GENERATED_BODY()
+};
+
+class IUIFrameworkWidgetWrapperInterface
+{
+	GENERATED_BODY()
+};
+
 /**
  * 
  */
@@ -38,6 +52,16 @@ public:
 	virtual bool CallRemoteFunction(UFunction* Function, void* Parameters, FOutParmRec* OutParms, FFrame* Stack) override;
 	//~ End UObject
 
+	TScriptInterface<IUIFrameworkWidgetWrapperInterface> AuthorityGetWrapper() const
+	{
+		return Wrapper;
+	}
+
+	void AuthoritySetWrapper(TScriptInterface<IUIFrameworkWidgetWrapperInterface> InWrapper)
+	{
+		Wrapper = InWrapper;
+	}
+
 	FUIFrameworkWidgetId GetWidgetId() const
 	{
 		return Id;
@@ -59,7 +83,9 @@ public:
 		return AuthorityParent;
 	}
 
-	virtual void AuthorityForEachChildren(const TFunctionRef<void(UUIFrameworkWidget*)>& Func) {}
+	virtual void AuthorityForEachChildren(const TFunctionRef<void(UUIFrameworkWidget*)>& Func)
+	{
+	}
 
 	//~ Local functions
 	UWidget* LocalGetUMGWidget() const
@@ -79,8 +105,12 @@ public:
 	void SetEnabled(bool bEnabled);
 
 protected:
-	virtual void AuthorityRemoveChild(UUIFrameworkWidget* Widget) {}
-	virtual void LocalOnUMGWidgetCreated() { }
+	virtual void AuthorityRemoveChild(UUIFrameworkWidget* Widget)
+	{
+	}
+	virtual void LocalOnUMGWidgetCreated()
+	{
+	}
 
 private:
 	UFUNCTION()
@@ -97,6 +127,10 @@ private:
 	UPROPERTY(Replicated, Transient, DuplicateTransient)
 	FUIFrameworkWidgetId Id = FUIFrameworkWidgetId::MakeNew();
 
+	//~ Authority
+	UPROPERTY(Transient, DuplicateTransient)
+	TScriptInterface<IUIFrameworkWidgetWrapperInterface> Wrapper;
+
 	//~ Authority and Local
 	UPROPERTY(Transient)
 	TObjectPtr<UUIFrameworkPlayerComponent> OwnerPlayerComponent = nullptr;
