The Presenter use a stack base on priority for UI input.

diff --git a/Engine/Plugins/Experimental/UIFramework/Source/Private/UIFPlayerComponent.cpp b/Engine/Plugins/Experimental/UIFramework/Source/Private/UIFPlayerComponent.cpp
--- a/Engine/Plugins/Experimental/UIFramework/Source/Private/UIFPlayerComponent.cpp
+++ b/Engine/Plugins/Experimental/UIFramework/Source/Private/UIFPlayerComponent.cpp
@@ -20,6 +20,23 @@
 /**
  *
  */
+
+void FUIFrameworkGameLayerSlotList::PreReplicatedRemove(const TArrayView<int32>& ChangedIndices, int32 FinalSize)
+{
+	check(Owner);
+	for (int32 Index : ChangedIndices)
+	{
+		FUIFrameworkGameLayerSlot& Slot = Entries[Index];
+		if (Slot.LocalIsAquiredWidgetValid())
+		{
+			if (Owner->Presenter)
+			{
+				Owner->Presenter->RemoveFromViewport(Slot.GetWidgetId());
+			}
+		}
+	}
+}
+
 void FUIFrameworkGameLayerSlotList::PostReplicatedChange(const TArrayView<int32>& ChangedIndices, int32 FinalSize)
 {
 	check(Owner);
diff --git a/Engine/Plugins/Experimental/UIFramework/Source/Private/UIFPresenter.cpp b/Engine/Plugins/Experimental/UIFramework/Source/Private/UIFPresenter.cpp
--- a/Engine/Plugins/Experimental/UIFramework/Source/Private/UIFPresenter.cpp
+++ b/Engine/Plugins/Experimental/UIFramework/Source/Private/UIFPresenter.cpp
@@ -7,6 +7,13 @@
 #include "Components/Widget.h"
 #include "GameFramework/PlayerController.h"
 
+UUIFrameworkGameViewportPresenter::FWidgetPair::FWidgetPair(UWidget* InWidget, FUIFrameworkWidgetId InWidgetId)
+	: UMGWidget(InWidget)
+	, WidgetId(InWidgetId)
+{
+
+}
+
 void UUIFrameworkGameViewportPresenter::AddToViewport(UWidget* UMGWidget, const FUIFrameworkGameLayerSlot& Slot)
 {
 	if (UGameViewportSubsystem* Subsystem = UGameViewportSubsystem::Get(GetOuterUUIFrameworkPlayerComponent()->GetWorld()))
@@ -23,5 +30,37 @@ void UUIFrameworkGameViewportPresenter::AddToViewport(UWidget* UMGWidget, const
 			check(LocalOwner);
 			Subsystem->AddWidgetForPlayer(UMGWidget, LocalOwner->GetLocalPlayer(), GameViewportWidgetSlot);
 		}
+		Widgets.Emplace(UMGWidget, Slot.GetWidgetId());
 	}
 }
+
+void UUIFrameworkGameViewportPresenter::RemoveFromViewport(FUIFrameworkWidgetId WidgetId)
+{
+	int32 IndexOf = Widgets.IndexOfByPredicate([WidgetId](const FWidgetPair& Other) { return Other.WidgetId == WidgetId; });
+	if (IndexOf != INDEX_NONE)
+	{
+		if (UGameViewportSubsystem* Subsystem = UGameViewportSubsystem::Get(GetOuterUUIFrameworkPlayerComponent()->GetWorld()))
+		{
+			Subsystem->RemoveWidget(Widgets[IndexOf].UMGWidget.Get());
+		}
+
+		Widgets.RemoveAtSwap(IndexOf);
+	}
+}
+
+void UUIFrameworkGameViewportPresenter::BeginDestroy()
+{
+	if (GetOuter() && !IsTemplate())
+	{
+		if (UGameViewportSubsystem* Subsystem = UGameViewportSubsystem::Get(GetOuterUUIFrameworkPlayerComponent()->GetWorld()))
+		{
+			for (FWidgetPair& Pair : Widgets)
+			{
+				Subsystem->RemoveWidget(Pair.UMGWidget.Get());
+			}
+		}
+	}
+	Widgets.Reset();
+
+	Super::BeginDestroy();
+}
\ No newline at end of file
diff --git a/Engine/Plugins/Experimental/UIFramework/Source/Public/UIFPlayerComponent.h b/Engine/Plugins/Experimental/UIFramework/Source/Public/UIFPlayerComponent.h
--- a/Engine/Plugins/Experimental/UIFramework/Source/Public/UIFPlayerComponent.h
+++ b/Engine/Plugins/Experimental/UIFramework/Source/Public/UIFPlayerComponent.h
@@ -78,6 +78,7 @@ struct UIFRAMEWORK_API FUIFrameworkGameLayerSlotList : public FFastArraySerializ
 
 public:
 	//~ Begin of FFastArraySerializer
+	void PreReplicatedRemove(const TArrayView<int32>& RemovedIndices, int32 FinalSize);
 	void PostReplicatedChange(const TArrayView<int32>& ChangedIndices, int32 FinalSize);
 	//~ End of FFastArraySerializer
 
diff --git a/Engine/Plugins/Experimental/UIFramework/Source/Public/UIFPresenter.h b/Engine/Plugins/Experimental/UIFramework/Source/Public/UIFPresenter.h
--- a/Engine/Plugins/Experimental/UIFramework/Source/Public/UIFPresenter.h
+++ b/Engine/Plugins/Experimental/UIFramework/Source/Public/UIFPresenter.h
@@ -4,10 +4,13 @@
 
 #include "UObject/Object.h"
 #include "UIFPlayerComponent.h"
+#include "Types/UIFWidgetId.h"
 
 #include "UIFPresenter.generated.h"
 
+struct FUIFrameworkGameLayerSlot;
 class FUIFrameworkModule;
+class UWidget;
 
 /**
  * 
@@ -21,6 +24,10 @@ public:
 	virtual void AddToViewport(UWidget* UMGWidget, const FUIFrameworkGameLayerSlot& Slot)
 	{
 
+	}
+	virtual void RemoveFromViewport(FUIFrameworkWidgetId WidgetId)
+	{
+
 	}
 };
 
@@ -35,4 +42,17 @@ class UIFRAMEWORK_API UUIFrameworkGameViewportPresenter : public UUIFrameworkPre
 
 public:
 	virtual void AddToViewport(UWidget* UMGWidget, const FUIFrameworkGameLayerSlot& Slot) override;
+	virtual void RemoveFromViewport(FUIFrameworkWidgetId WidgetId) override;
+
+	virtual void BeginDestroy() override;
+
+private:
+	struct FWidgetPair
+	{
+		FWidgetPair() = default;
+		FWidgetPair(UWidget* Widget, FUIFrameworkWidgetId WidgetId);
+		TWeakObjectPtr<UWidget> UMGWidget;
+		FUIFrameworkWidgetId WidgetId;
+	};
+	TArray<FWidgetPair> Widgets;
 };
