UIFramework: Add debug comments to track replication issue.
UIFramework.PrintWidgetTree logs the replicated data. PrintWidgetsChildren logs the runtime data.

diff --git a/Engine/Plugins/Experimental/UIFramework/Source/Private/Types/UIFWidgetTree.cpp b/Engine/Plugins/Experimental/UIFramework/Source/Private/Types/UIFWidgetTree.cpp
--- a/Engine/Plugins/Experimental/UIFramework/Source/Private/Types/UIFWidgetTree.cpp
+++ b/Engine/Plugins/Experimental/UIFramework/Source/Private/Types/UIFWidgetTree.cpp
@@ -5,6 +5,11 @@
 #include "UIFPlayerComponent.h"
 #include "UIFWidget.h"
 
+#if UE_UIFRAMEWORK_WITH_DEBUG
+#include "Components/PanelWidget.h"
+#include "Components/Widget.h"
+#endif
+
 #include "Engine/ActorChannel.h"
 #include "Engine/Engine.h"
 #include "Engine/NetDriver.h"
@@ -26,11 +31,39 @@ namespace UE::UIFramework::Private
 		}
 	}
 
+	void PrintWidgetTree()
+	{
+		for (FUIFrameworkWidgetTree* Tree : GTrees)
+		{
+			Tree->LogTree();
+		}
+	}
+
+	void PrintWidgetsChildren()
+	{
+		for (FUIFrameworkWidgetTree* Tree : GTrees)
+		{
+			Tree->LogWidgetsChildren();
+		}
+	}
+
 	static FAutoConsoleCommand CCmdTestWidgetTree(
 		TEXT("UIFramework.TestWidgetTree"),
-		TEXT("Test if all containers are properly setup"),
+		TEXT("Test if all containers are properly setup."),
 		FConsoleCommandDelegate::CreateStatic(TestWidgetTree),
 		ECVF_Cheat);
+
+	static FAutoConsoleCommand CCmdPrintWidgetTree(
+		TEXT("UIFramework.PrintWidgetTree"),
+		TEXT("Print/Log all the widgets in the widgettree as they are store in the tree."),
+		FConsoleCommandDelegate::CreateStatic(PrintWidgetTree),
+		ECVF_Cheat);
+
+	static FAutoConsoleCommand CCmdPrintWidgetsChildren(
+		TEXT("UIFramework.PrintWidgetsChildren"),
+		TEXT("Print/Log all the widgets in the treetree as they are used at runtime."),
+		FConsoleCommandDelegate::CreateStatic(PrintWidgetsChildren),
+		ECVF_Cheat);
 }
 #endif
 
@@ -340,6 +373,7 @@ void FUIFrameworkWidgetTree::AuthorityTest() const
 {
 	if (ReplicatedOwner == nullptr || !ReplicatedOwner->HasAuthority())
 	{
+		UE_LOG(LogUIFramework, Warning, TEXT("Can't run the WidgetTree tests on local."));
 		return;
 	}
 
@@ -413,7 +447,109 @@ void FUIFrameworkWidgetTree::AuthorityTest() const
 				ensureAlwaysMsgf(false, TEXT("Widget no in the map"));
 			}
 		}
-		
+	}
+	UE_LOG(LogUIFramework, Log, TEXT("Completed WidgetTree tests"));
+}
+
+void FUIFrameworkWidgetTree::LogTreeInternal(const FUIFrameworkWidgetTreeEntry& ParentEntry, FString Spaces) const
+{
+	Spaces.AppendChar(TEXT('>'));
+	for (FUIFrameworkWidgetTreeEntry Entry : Entries)
+	{
+		if (ParentEntry.ChildId == Entry.ParentId)
+		{
+			UE_LOG(LogUIFramework, Log, TEXT("%s%s"), *Spaces, Entry.Child ? *Entry.Child->GetName() : TEXT("<unknow>"));
+			LogTreeInternal(Entry, Spaces);
+		}
+	}
+}
+
+void FUIFrameworkWidgetTree::LogTree() const
+{
+	UE_LOG(LogUIFramework, Log, TEXT("WidgetTree for owner '%s'."), ReplicatedOwner ? *ReplicatedOwner->GetPathName() : TEXT("<none>"));
+
+	for (FUIFrameworkWidgetTreeEntry Entry : Entries)
+	{
+		if (Entry.ParentId.IsRoot())
+		{
+			UE_LOG(LogUIFramework, Log, TEXT(">%s"), Entry.Child ? *Entry.Child->GetName() : TEXT("<unknow>"));
+			LogTreeInternal(Entry, TEXT(">"));
+		}
+	}
+}
+
+void FUIFrameworkWidgetTree::AuthorityLogWidgetsChildrenInternal(UUIFrameworkWidget* Widget, FString Spaces) const
+{
+	Spaces.AppendChar(TEXT('>'));
+	Widget->AuthorityForEachChildren([this, &Spaces](UUIFrameworkWidget* Child)
+		{
+			UE_LOG(LogUIFramework, Log, TEXT("%s%s"), *Spaces, Child ? *Child->GetName() : TEXT("<unknow>"));
+			if (Child)
+			{
+				AuthorityLogWidgetsChildrenInternal(Child, Spaces);
+			}
+		});
+}
+
+void FUIFrameworkWidgetTree::LocalLogWidgetsChildrenInternal(const UWidget* UMGWidget, FString Spaces) const
+{
+	Spaces.AppendChar(TEXT('>'));
+	if (const UPanelWidget* UMGPanel = Cast<UPanelWidget>(UMGWidget))
+	{
+		const int32 Count = UMGPanel->GetChildrenCount();
+		for (int32 Index = 0; Index < Count; ++Index)
+		{
+			UWidget* ChildWidget = UMGPanel->GetChildAt(Index);
+			UE_LOG(LogUIFramework, Log, TEXT("%s%s"), *Spaces, ChildWidget ? *ChildWidget->GetName() : TEXT("<unknow>"));
+			if (ChildWidget)
+			{
+				LocalLogWidgetsChildrenInternal(ChildWidget, Spaces);
+			}
+		}
+	}
+}
+
+void FUIFrameworkWidgetTree::LogWidgetsChildren() const
+{
+	if (ReplicatedOwner == nullptr)
+	{
+		return;
+	}
+	UE_LOG(LogUIFramework, Log, TEXT("Widgets Child for owner '%s'."), *ReplicatedOwner->GetPathName());
+
+	if (ReplicatedOwner->HasAuthority())
+	{
+		UE_LOG(LogUIFramework, Log, TEXT("== Authority =="));
+
+		for (FUIFrameworkWidgetTreeEntry Entry : Entries)
+		{
+			if (Entry.ParentId.IsRoot())
+			{
+				UE_LOG(LogUIFramework, Log, TEXT(">%s"), Entry.Child ? *Entry.Child->GetName() : TEXT("<unknow>"));
+				if (Entry.Child)
+				{
+					AuthorityLogWidgetsChildrenInternal(Entry.Child, TEXT(">"));
+				}
+			}
+		}
+	}
+	else
+	{
+		UE_LOG(LogUIFramework, Log, TEXT("== Local =="));
+
+		for (FUIFrameworkWidgetTreeEntry Entry : Entries)
+		{
+			if (Entry.ParentId.IsRoot())
+			{
+				UWidget* UMGWidget = Entry.Child ? Entry.Child->LocalGetUMGWidget() : nullptr;
+
+				UE_LOG(LogUIFramework, Log, TEXT(">%s"), UMGWidget ? *UMGWidget->GetName() : TEXT("<unknow>"));
+				if (UMGWidget)
+				{
+					LocalLogWidgetsChildrenInternal(UMGWidget, TEXT(">"));
+				}
+			}
+		}
 	}
 }
 #endif
diff --git a/Engine/Plugins/Experimental/UIFramework/Source/Public/Types/UIFWidgetTree.h b/Engine/Plugins/Experimental/UIFramework/Source/Public/Types/UIFWidgetTree.h
--- a/Engine/Plugins/Experimental/UIFramework/Source/Public/Types/UIFWidgetTree.h
+++ b/Engine/Plugins/Experimental/UIFramework/Source/Public/Types/UIFWidgetTree.h
@@ -10,7 +10,7 @@
 #include "UIFWidgetTree.generated.h"
 
 #ifndef UE_UIFRAMEWORK_WITH_DEBUG
-	#define UE_UIFRAMEWORK_WITH_DEBUG !(UE_BUILD_SHIPPING || UE_BUILD_TEST)
+	#define UE_UIFRAMEWORK_WITH_DEBUG !(UE_BUILD_SHIPPING)
 #endif
 
 struct FReplicationFlags;
@@ -20,6 +20,10 @@ class FOutBunch;
 class UActorChannel;
 class UUIFrameworkWidget;
 
+#if UE_UIFRAMEWORK_WITH_DEBUG
+class UWidget;
+#endif
+
 class IUIFrameworkWidgetTreeOwner
 {
 public:
@@ -131,6 +135,8 @@ public:
 
 #if UE_UIFRAMEWORK_WITH_DEBUG
 	void AuthorityTest() const;
+	void LogTree() const;
+	void LogWidgetsChildren() const;
 #endif
 
 	DECLARE_MULTICAST_DELEGATE_OneParam(FUIFrameworkWidgetDelegate, UUIFrameworkWidget*);
@@ -142,6 +148,12 @@ private:
 	void AuthorityAddChildRecursiveInternal(UUIFrameworkWidget* Widget);
 	bool AuthorityRemoveChildRecursiveInternal(UUIFrameworkWidget* Widget);
 
+#if UE_UIFRAMEWORK_WITH_DEBUG
+	void LogTreeInternal(const FUIFrameworkWidgetTreeEntry& Entry, FString Spaces) const;
+	void AuthorityLogWidgetsChildrenInternal(UUIFrameworkWidget* Widget, FString Spaces) const;
+	void LocalLogWidgetsChildrenInternal(const UWidget* Widget, FString Spaces) const;
+#endif
+
 private:
 	UPROPERTY()
 	TArray<FUIFrameworkWidgetTreeEntry> Entries;
