UIFramework: Update the widget tree. Only keep a reference of widgets that are on the local player. Add architecture to support owners other than the PlayerController. Support, temporarily, a widget that moves from player to player.

diff --git a/Engine/Plugins/Experimental/UIFramework/Source/Private/Types/UIFParentWidget.cpp b/Engine/Plugins/Experimental/UIFramework/Source/Private/Types/UIFParentWidget.cpp
--- a/Engine/Plugins/Experimental/UIFramework/Source/Private/Types/UIFParentWidget.cpp
+++ b/Engine/Plugins/Experimental/UIFramework/Source/Private/Types/UIFParentWidget.cpp
@@ -34,3 +34,8 @@ UUIFrameworkPlayerComponent* FUIFrameworkParentWidget::AsPlayerComponent() const
 {
 	return CastChecked<UUIFrameworkPlayerComponent>(Parent);
 }
+
+bool FUIFrameworkParentWidget::operator== (const UUIFrameworkWidget* Other) const
+{
+	return Other == Parent && bIsParentAWidget;
+}
\ No newline at end of file
diff --git a/Engine/Plugins/Experimental/UIFramework/Source/Private/Types/UIFWidgetTree.cpp b/Engine/Plugins/Experimental/UIFramework/Source/Private/Types/UIFWidgetTree.cpp
--- a/Engine/Plugins/Experimental/UIFramework/Source/Private/Types/UIFWidgetTree.cpp
+++ b/Engine/Plugins/Experimental/UIFramework/Source/Private/Types/UIFWidgetTree.cpp
@@ -16,6 +16,7 @@
 #include "GameFramework/Actor.h"
 #include "GameFramework/PlayerController.h"
 #include "Net/UnrealNetwork.h"
+#include "UObject/Package.h"
 
 #include UE_INLINE_GENERATED_CPP_BY_NAME(UIFWidgetTree)
 
@@ -89,6 +90,11 @@ bool FUIFrameworkWidgetTreeEntry::IsChildValid() const
 	return Child && Child->GetWidgetId() == ChildId;
 }
 
+FString FUIFrameworkWidgetTreeEntry::GetDebugString()
+{
+	return FString::Printf(TEXT("%s %s"), Parent ? *Parent->GetName() : TEXT("<none>"), Child ? *Child->GetName() : TEXT("<none>"));
+}
+
 /**
  *
  */
@@ -209,6 +215,21 @@ void FUIFrameworkWidgetTree::AuthorityAddWidget(UUIFrameworkWidget* Parent, UUIF
 
 void FUIFrameworkWidgetTree::AuthorityAddChildInternal(UUIFrameworkWidget* Parent, UUIFrameworkWidget* Child, bool bFirst)
 {
+	Child->WidgetTreeOwner = Owner;
+	check(bFirst || Parent->WidgetTreeOwner == Owner);
+
+	bool bOuterIsDifferent = Child->GetOuter() != ReplicatedOwner;
+	if (bOuterIsDifferent)
+	{
+		//if (Child->GetOuter() == GetTransientPackage())
+		{
+			//If the outer is the transient package, then there are no replication owner yetand it safe to rename it with the correct new outer.
+			//If the Outer is not the transient, then the widget got replicated with another player.There are no "reset" and we should duplicate and delete the object.
+			//For now only do the rename.It works but it is not the best.
+			Child->Rename(nullptr, ReplicatedOwner);
+		}
+	}
+
 	if (int32* PreviousEntryIndexPtr = AuthorityIndexByWidgetMap.Find(Child))
 	{
 		check(Entries.IsValidIndex(*PreviousEntryIndexPtr));
@@ -219,6 +240,11 @@ void FUIFrameworkWidgetTree::AuthorityAddChildInternal(UUIFrameworkWidget* Paren
 			PreviousEntry = FUIFrameworkWidgetTreeEntry(Parent, Child);
 			MarkItemDirty(PreviousEntry);
 		}
+
+		if (bOuterIsDifferent)
+		{
+			AuthorityAddChildRecursiveInternal(Child);
+		}
 	}
 	else
 	{
@@ -255,22 +281,42 @@ void FUIFrameworkWidgetTree::AuthorityAddChildRecursiveInternal(UUIFrameworkWidg
 		});
 }
 
-void FUIFrameworkWidgetTree::AuthorityRemoveWidget(UUIFrameworkWidget* Widget)
+void FUIFrameworkWidgetTree::AuthorityRemoveWidgetAndChildren(UUIFrameworkWidget* Widget)
 {
-	check(Widget);
-
-	if (AuthorityRemoveChildRecursiveInternal(Widget))
+	if (ensure(Widget))
 	{
-		MarkArrayDirty();
-		AuthorityOnWidgetRemoved.Broadcast(Widget);
+		if (AuthorityRemoveChildRecursiveInternal(Widget))
+		{
+			MarkArrayDirty();
+			AuthorityOnWidgetRemoved.Broadcast(Widget);
+		}
 	}
 }
 
+//void FUIFrameworkWidgetTree::AuthorityReplaceWidget(UUIFrameworkWidget* OldWidget, UUIFrameworkWidget* NewWidget)
+//{
+//	bool bHasReplaced = false;
+//	for (FUIFrameworkWidgetTreeEntry& Entry : Entries)
+//	{
+//		if (Entry.Child == OldWidget)
+//		{
+//			bHasReplaced = true;
+//		}
+//		if (Entry.Parent == OldWidget)
+//		{
+//			bHasReplaced = true;
+//		}
+//	}
+//
+//	if (bHasReplaced)
+//	{
+//
+//	}
+//}
+
 void FUIFrameworkWidgetTree::LocalRemoveRoot(const UUIFrameworkWidget* Widget)
 {
-	check(Widget);
-
-	if (ensure(Owner))
+	if (ensure(Owner && Widget))
 	{
 		Owner->LocalRemoveWidgetRootFromTree(Widget);
 	}
@@ -278,6 +324,7 @@ void FUIFrameworkWidgetTree::LocalRemoveRoot(const UUIFrameworkWidget* Widget)
 
 bool FUIFrameworkWidgetTree::AuthorityRemoveChildRecursiveInternal(UUIFrameworkWidget* Widget)
 {
+	Widget->WidgetTreeOwner = nullptr;
 	if (int32* PreviousEntryIndexPtr = AuthorityIndexByWidgetMap.Find(Widget))
 	{
 		check(Entries.IsValidIndex(*PreviousEntryIndexPtr));
diff --git a/Engine/Plugins/Experimental/UIFramework/Source/Private/Types/UIFWidgetTreeOwner.cpp b/Engine/Plugins/Experimental/UIFramework/Source/Private/Types/UIFWidgetTreeOwner.cpp
new file mode 100644
--- /dev/null
+++ b/Engine/Plugins/Experimental/UIFramework/Source/Private/Types/UIFWidgetTreeOwner.cpp
@@ -0,0 +1,10 @@
+// Copyright Epic Games, Inc. All Rights Reserved.
+
+#include "Types/UIFWidgetTreeOwner.h"
+#include "Types/UIFWidgetTree.h"
+
+
+FUIFrameworkWidgetTree& IUIFrameworkWidgetTreeOwner::GetWidgetTree()
+{
+	PURE_VIRTUAL(IUIFrameworkWidgetTreeOwner::GetWidgetTree, static FUIFrameworkWidgetTree Tmp; return Tmp;)
+}
\ No newline at end of file
diff --git a/Engine/Plugins/Experimental/UIFramework/Source/Private/UIFModule.cpp b/Engine/Plugins/Experimental/UIFramework/Source/Private/UIFModule.cpp
--- a/Engine/Plugins/Experimental/UIFramework/Source/Private/UIFModule.cpp
+++ b/Engine/Plugins/Experimental/UIFramework/Source/Private/UIFModule.cpp
@@ -13,78 +13,129 @@ namespace UE::UIFramework::Private
 	static TSubclassOf<UUIFrameworkPresenter> Director;
 }
 
-UUIFrameworkWidget* FUIFrameworkModule::AuthorityAttachWidget(UUIFrameworkPlayerComponent* ReplicationOwner, FUIFrameworkParentWidget Parent, UUIFrameworkWidget* Child)
+// UIFrameworkWidget has 
+//	1. a WidgetTree owner (UUIFrameworkPlayerComponent/IUIFrameworkWidgetTreeOwner) [authority/local]
+//		set by the WidgetTree when they are added or removed from the tree
+//  2. a replication Outer / GetOuter (APlayerController or transient if not added yet) [authority/local]
+//		set by the WidgetTree when they are added
+//  3. a wrapper object (VerseUIWidget) [authority]
+//		set by the VerseUIWidget on creation
+//	4. a parent widget [authority]
+//		set by FUIFrameworkModule::AuthorityAttachWidget or FUIFrameworkModule::AuthorityDetachWidgetFromParent
+
+UUIFrameworkWidget* FUIFrameworkModule::AuthorityAttachWidget(FUIFrameworkParentWidget Parent, UUIFrameworkWidget* Child)
 {
-	check(Child);
-	check(Parent.IsParentValid());
-
-	if (Child->AuthorityGetParent().IsParentValid())
-	{
-		AuthorityDetachWidgetFromParent(Child);
-	}
-
-	UObject* ParentOuter = Parent.IsPlayerComponent() ? Parent.AsPlayerComponent()->GetOuter() : Parent.AsWidget()->GetOuter();
-	if (ParentOuter != Child->GetOuter())
+	if (ensure(Child && Parent.IsParentValid()))
 	{
-		if (Child->GetOuter() == GetTransientPackage())
+		// If re-parenting on itself or already attach
+		if (Parent == Child || Child->AuthorityGetParent() == Parent)
 		{
-			// Recursive Rename to change the outer
-			// To the operation recursively to child
+			return Child;
 		}
-	}
 
-	bool bDifferentReplicationOwner = ReplicationOwner != Child->OwnerPlayerComponent;
+		// Remove from previous parent.
+		if (Child->AuthorityGetParent().IsParentValid())
+		{
+			bool bHasSameWidgetTree = false;
+			if (Parent.IsPlayerComponent())
+			{
+				bHasSameWidgetTree = &Parent.AsPlayerComponent()->GetWidgetTree() == Child->GetWidgetTree();
+			}
+			else
+			{
+				check(Parent.IsWidget());
+				bHasSameWidgetTree = Parent.AsWidget()->GetWidgetTree() == Child->GetWidgetTree();
+			}
+			AuthorityDetachWidgetFromParentInternal(Child, bHasSameWidgetTree);
+		}
 
-	Child->AuthorityParent = Parent;
-	Child->OwnerPlayerComponent = ReplicationOwner;
+		Child->AuthorityParent = Parent;
 
-	if (ReplicationOwner)
-	{
-		if (Child->AuthorityParent.IsWidget())
+		// If the parent in the WidgetTree or the parent a root, then add it to the WidgetTree
+		if (Parent.IsPlayerComponent())
 		{
-			ReplicationOwner->GetWidgetTree().AuthorityAddWidget(Child->AuthorityParent.AsWidget(), Child);
+			Parent.AsPlayerComponent()->GetWidgetTree().AuthorityAddRoot(Child);
 		}
-		else
+		else if (FUIFrameworkWidgetTree* WidgetTree = Parent.AsWidget()->GetWidgetTree())
 		{
-			check(Child->AuthorityParent.IsPlayerComponent());
-			ReplicationOwner->GetWidgetTree().AuthorityAddRoot(Child);
+			WidgetTree->AuthorityAddWidget(Child->AuthorityParent.AsWidget(), Child);
 		}
 	}
 
-	if (bDifferentReplicationOwner)
-	{
-		AuthoritySetParentReplicationOwnerRecursive(Child);
-	}
 	return Child;
 }
 
-void FUIFrameworkModule::AuthoritySetParentReplicationOwnerRecursive(UUIFrameworkWidget* Widget)
+//void FUIFrameworkModule::AuthoritySetParentReplicationOwnerRecursive(UUIFrameworkWidget* Widget)
+//{
+	//Widget->AuthorityForEachChildren([Widget](UUIFrameworkWidget* Child)
+	//	{
+	//		if (Child != nullptr)
+	//		{
+	//			check(Child->AuthorityGetParent().IsWidget() && Child->AuthorityGetParent().AsWidget() == Widget);
+	//			//Child->OwnerPlayerComponent = Widget->OwnerPlayerComponent;
+	//			Child->AuthorityParent = FUIFrameworkParentWidget(Widget);
+	//			AuthoritySetParentReplicationOwnerRecursive(Child);
+	//		}
+	//	});
+//}
+
+//UUIFrameworkWidget* FUIFrameworkModule::AuthorityRenameRecursive(UUIFrameworkPlayerComponent* ReplicationOwner, UUIFrameworkWidget* Widget, UObject* NewOuter)
+//{
+	//if (NewOuter != Widget->GetOuter())
+	//{
+	//	if (Widget->GetOuter() == GetTransientPackage())
+	//	{
+	//		// If the outer is the transient package, then there are no replication owner yet and it safe to just rename it with the correct new owner.
+	//		Widget->Rename(nullptr, NewOuter);
+	//	}
+	//	else
+	//	{
+	//		// The widget change owner, we need to create a new one to replicate it correctly.
+	//		UUIFrameworkWidget* NewWidget = DuplicateObject<UUIFrameworkWidget>(Widget, NewOuter);
+	//		// Replace all instance in the widget tree
+	//		ReplicationOwner->GetWidgetTree()->AuthorityReplaceWidget(Widget, NewWidget);
+	//		// Notify it's owner
+	//		if (Widget->AuthorityGetWrapper())
+	//		{
+	//			Widget->AuthorityGetWrapper()->ReplaceWidget(Widget, NewWidget);
+	//		}
+	//		Widget = NewWidget;
+	//	}
+
+	//	Widget->AuthorityForEachChildren([ReplicationOwner, NewOuter](UUIFrameworkWidget* Child)
+	//		{
+	//			if (Child != nullptr)
+	//			{
+	//				AuthorityRenameRecursive(ReplicationOwner, Child, NewOuter);
+	//			}
+	//		});
+	//}
+//	return Widget;
+//}
+
+bool FUIFrameworkModule::AuthorityCanWidgetBeAttached(FUIFrameworkParentWidget Parent, UUIFrameworkWidget* Child)
 {
-	Widget->AuthorityForEachChildren([Widget](UUIFrameworkWidget* Child)
-		{
-			if (Child != nullptr)
-			{
-				check(Child->AuthorityGetParent().IsWidget() && Child->AuthorityGetParent().AsWidget() == Widget);
-				Child->OwnerPlayerComponent = Widget->OwnerPlayerComponent;
-				Child->AuthorityParent = FUIFrameworkParentWidget(Widget);
-				AuthoritySetParentReplicationOwnerRecursive(Child);
-			}
-		});
+	return Child && Parent.IsParentValid() && Parent != Child;
 }
 
-bool FUIFrameworkModule::AuthorityCanWidgetBeAttached(UUIFrameworkPlayerComponent* ReplicationOwner, UUIFrameworkWidget* Parent, UUIFrameworkWidget* Child)
+void FUIFrameworkModule::AuthorityDetachWidgetFromParent(UUIFrameworkWidget* Child)
 {
-	return true;
+	AuthorityDetachWidgetFromParentInternal(Child, false);
 }
 
-void FUIFrameworkModule::AuthorityDetachWidgetFromParent(UUIFrameworkWidget* Child)
+void FUIFrameworkModule::AuthorityDetachWidgetFromParentInternal(UUIFrameworkWidget* Child, bool bTemporary)
 {
 	check(Child);
-	if (Child->OwnerPlayerComponent)
+
+	// If it's in the WidgetTree, we need to remove it.
+	if (!bTemporary && Child->WidgetTreeOwner)
 	{
-		Child->OwnerPlayerComponent->GetWidgetTree().AuthorityRemoveWidget(Child);
+		//bTemporary: we do not want to remove and re-add the same widget if it's not needed. 
+		//Removing them would cause the local to recreate them instead of re-parenting them.
+		Child->WidgetTreeOwner->GetWidgetTree().AuthorityRemoveWidgetAndChildren(Child);
 	}
 
+	// Notify the widget that we removed a child.
 	if (Child->AuthorityGetParent().IsParentValid())
 	{
 		if (Child->AuthorityGetParent().IsWidget())
@@ -97,6 +148,8 @@ void FUIFrameworkModule::AuthorityDetachWidgetFromParent(UUIFrameworkWidget* Chi
 			Child->AuthorityGetParent().AsPlayerComponent()->AuthorityRemoveChild(Child);
 		}
 	}
+
+	Child->AuthorityParent = FUIFrameworkParentWidget();
 }
 
 void FUIFrameworkModule::SetPresenterClass(TSubclassOf<UUIFrameworkPresenter> InDirector)
diff --git a/Engine/Plugins/Experimental/UIFramework/Source/Private/UIFPlayerComponent.cpp b/Engine/Plugins/Experimental/UIFramework/Source/Private/UIFPlayerComponent.cpp
--- a/Engine/Plugins/Experimental/UIFramework/Source/Private/UIFPlayerComponent.cpp
+++ b/Engine/Plugins/Experimental/UIFramework/Source/Private/UIFPlayerComponent.cpp
@@ -160,17 +160,9 @@ void UUIFrameworkPlayerComponent::AddWidget(FUIFrameworkGameLayerSlot InEntry)
 	}
 	else
 	{
-		UUIFrameworkPlayerComponent* PreviousOwner = InEntry.AuthorityGetWidget()->GetPlayerComponent();
-		if (PreviousOwner != nullptr && PreviousOwner != this)
-		{
-			FFrame::KismetExecutionMessage(TEXT("The widget was created for another player. It can't be added."), ELogVerbosity::Warning, "InvalidPlayerParent");
-		}
-		else
-		{
-			// Reset the widget to make sure the id is set and it may have been duplicated during the attach
-			InEntry.AuthoritySetWidget(FUIFrameworkModule::AuthorityAttachWidget(this, this, InEntry.AuthorityGetWidget()));
-			RootList.AddEntry(InEntry);
-		}
+		// Reset the widget to make sure the id is set and it may have been duplicated during the attach
+		InEntry.AuthoritySetWidget(FUIFrameworkModule::AuthorityAttachWidget(this, InEntry.AuthorityGetWidget()));
+		RootList.AddEntry(InEntry);
 	}
 }
 
@@ -185,15 +177,16 @@ void UUIFrameworkPlayerComponent::RemoveWidget(UUIFrameworkWidget* Widget)
 	}
 	else
 	{
-		UUIFrameworkPlayerComponent* PreviousOwner = Widget->GetPlayerComponent();
-		if (PreviousOwner != this)
+		if (Widget->GetWidgetTreeOwner() != this)
 		{
-			FFrame::KismetExecutionMessage(TEXT("The widget was created for another player. It can't be removed on this player."), ELogVerbosity::Warning, "InvalidPlayerParentOnRemovedWidget");
+			FFrame::KismetExecutionMessage(TEXT("The widget was not added on this player. It can't be removed on this player."), ELogVerbosity::Warning, "InvalidPlayerParentOnRemovedWidget");
 		}
 		else
 		{
-			FUIFrameworkModule::AuthorityDetachWidgetFromParent(Widget);
-			RootList.RemoveEntry(Widget);
+			if (RootList.RemoveEntry(Widget))
+			{
+				FUIFrameworkModule::AuthorityDetachWidgetFromParent(Widget);
+			}
 		}
 	}
 }
@@ -204,6 +197,17 @@ void UUIFrameworkPlayerComponent::AuthorityRemoveChild(UUIFrameworkWidget* Widge
 	RootList.RemoveEntry(Widget);
 }
 
+FUIFrameworkWidgetTree& UUIFrameworkPlayerComponent::GetWidgetTree()
+{
+	return WidgetTree;
+}
+
+FUIFrameworkWidgetOwner UUIFrameworkPlayerComponent::GetWidgetOwner() const
+{
+	FUIFrameworkWidgetOwner Owner;
+	Owner.PlayerController = GetPlayerController();
+	return Owner;
+}
 
 void UUIFrameworkPlayerComponent::TickComponent(float DeltaTime, enum ELevelTick TickType, FActorComponentTickFunction* ThisTickFunction)
 {
@@ -346,7 +350,7 @@ void UUIFrameworkPlayerComponent::ServerRemoveWidgetRootFromTree_Implementation(
 {
 	if (UUIFrameworkWidget* Widget = GetWidgetTree().FindWidgetById(WidgetId))
 	{
-		GetWidgetTree().AuthorityRemoveWidget(Widget);
+		GetWidgetTree().AuthorityRemoveWidgetAndChildren(Widget);
 	}
 }
 
diff --git a/Engine/Plugins/Experimental/UIFramework/Source/Private/UIFWidget.cpp b/Engine/Plugins/Experimental/UIFramework/Source/Private/UIFWidget.cpp
--- a/Engine/Plugins/Experimental/UIFramework/Source/Private/UIFWidget.cpp
+++ b/Engine/Plugins/Experimental/UIFramework/Source/Private/UIFWidget.cpp
@@ -30,7 +30,7 @@ int32 UUIFrameworkWidget::GetFunctionCallspace(UFunction* Function, FFrame* Stac
 		// This handles absorbing authority/cosmetic
 		return GEngine->GetGlobalFunctionCallspace(Function, this, Stack);
 	}
-	if (AActor* OwnerActor = GetPlayerComponent() ? GetPlayerComponent()->GetOwner() : nullptr)
+	if (AActor* OwnerActor = Cast<AActor>(GetOuter()))
 	{
 		return OwnerActor->GetFunctionCallspace(Function, Stack);
 	}
@@ -42,7 +42,7 @@ bool UUIFrameworkWidget::CallRemoteFunction(UFunction* Function, void* Parameter
 	check(!HasAnyFlags(RF_ClassDefaultObject));
 
 	bool bProcessed = false;
-	AActor* OwnerActor = GetPlayerComponent() ? GetPlayerComponent()->GetOwner() : nullptr;
+	AActor* OwnerActor = Cast<AActor>(GetOuter());
 	FWorldContext* const Context = OwnerActor ? GEngine->GetWorldContextFromWorld(OwnerActor->GetWorld()) : nullptr;
 	if (Context)
 	{
@@ -70,13 +70,18 @@ void UUIFrameworkWidget::GetLifetimeReplicatedProps(TArray<FLifetimeProperty>& O
 	DOREPLIFETIME_WITH_PARAMS_FAST(ThisClass, Visibility, Params);
 }
 
+FUIFrameworkWidgetTree* UUIFrameworkWidget::GetWidgetTree() const
+{
+	return WidgetTreeOwner ? &WidgetTreeOwner->GetWidgetTree() : nullptr;
+}
+
 void UUIFrameworkWidget::LocalAddChild(FUIFrameworkWidgetId ChildId)
 {
 	// By default we should remove the widget from its previous parent.
 	//Adding a widget to a new slot will automatically remove it from its previous parent.
-	if (OwnerPlayerComponent)
+	if (FUIFrameworkWidgetTree* WidgetTree = GetWidgetTree())
 	{
-		if (UUIFrameworkWidget* Widget = OwnerPlayerComponent->GetWidgetTree().FindWidgetById(ChildId))
+		if (UUIFrameworkWidget* Widget = WidgetTree->FindWidgetById(ChildId))
 		{
 			if (UWidget* UMGWidget = Widget->LocalGetUMGWidget())
 			{
@@ -86,14 +91,30 @@ void UUIFrameworkWidget::LocalAddChild(FUIFrameworkWidgetId ChildId)
 	}
 }
 
-void UUIFrameworkWidget::LocalCreateUMGWidget(UUIFrameworkPlayerComponent* InOwner)
+void UUIFrameworkWidget::LocalCreateUMGWidget(TNonNullPtr<IUIFrameworkWidgetTreeOwner> InOwner)
 {
-	OwnerPlayerComponent = InOwner;
+	WidgetTreeOwner = InOwner;
 	if (UClass* Class = WidgetClass.Get())
 	{
 		if (Class->IsChildOf(UUserWidget::StaticClass()))
 		{
-			LocalUMGWidget = CreateWidget(OwnerPlayerComponent->GetPlayerController(), Class);
+			FUIFrameworkWidgetOwner UserWidgetOwner = WidgetTreeOwner->GetWidgetOwner();
+			if (UserWidgetOwner.PlayerController)
+			{
+				LocalUMGWidget = CreateWidget(UserWidgetOwner.PlayerController, Class);
+			}
+			else if (UserWidgetOwner.GameInstance)
+			{
+				LocalUMGWidget = CreateWidget(UserWidgetOwner.GameInstance, Class);
+			}
+			else if (UserWidgetOwner.World)
+			{
+				LocalUMGWidget = CreateWidget(UserWidgetOwner.World, Class);
+			}
+			else
+			{
+				ensureAlwaysMsgf(false, TEXT("There are no valid UserWidget owner."));
+			}
 		}
 		else
 		{
@@ -114,7 +135,7 @@ void UUIFrameworkWidget::LocalDestroyUMGWidget()
 		LocalUMGWidget->ReleaseSlateResources(true);
 	}
 	LocalUMGWidget = nullptr;
-	OwnerPlayerComponent = nullptr;
+	WidgetTreeOwner = nullptr;
 }
 
 
diff --git a/Engine/Plugins/Experimental/UIFramework/Source/Private/Widgets/UIFButton.cpp b/Engine/Plugins/Experimental/UIFramework/Source/Private/Widgets/UIFButton.cpp
--- a/Engine/Plugins/Experimental/UIFramework/Source/Private/Widgets/UIFButton.cpp
+++ b/Engine/Plugins/Experimental/UIFramework/Source/Private/Widgets/UIFButton.cpp
@@ -46,16 +46,6 @@ void UUIFrameworkButton::SetContent(FUIFrameworkSimpleSlot InEntry)
 		{
 			FUIFrameworkModule::AuthorityDetachWidgetFromParent(Slot.AuthorityGetWidget());
 		}
-
-		if (InEntry.AuthorityGetWidget())
-		{
-			UUIFrameworkPlayerComponent* PreviousOwner = InEntry.AuthorityGetWidget()->GetPlayerComponent();
-			if (PreviousOwner != nullptr && GetPlayerComponent() != nullptr && PreviousOwner != GetPlayerComponent())
-			{
-				Slot.AuthoritySetWidget(nullptr);
-				FFrame::KismetExecutionMessage(TEXT("The widget was created for another player. It can't be added."), ELogVerbosity::Warning, "InvalidPlayerParent");
-			}
-		}
 	}
 
 	Slot = InEntry;
@@ -63,7 +53,7 @@ void UUIFrameworkButton::SetContent(FUIFrameworkSimpleSlot InEntry)
 	if (bWidgetIsDifferent && Slot.AuthorityGetWidget())
 	{
 		// Reset the widget to make sure the id is set and it may have been duplicated during the attach
-		Slot.AuthoritySetWidget(FUIFrameworkModule::AuthorityAttachWidget(GetPlayerComponent(), this, Slot.AuthorityGetWidget()));
+		Slot.AuthoritySetWidget(FUIFrameworkModule::AuthorityAttachWidget(this, Slot.AuthorityGetWidget()));
 	}
 
 	MARK_PROPERTY_DIRTY_FROM_NAME(ThisClass, Slot, this);
@@ -92,16 +82,16 @@ void UUIFrameworkButton::LocalOnUMGWidgetCreated()
 {
 	Super::LocalOnUMGWidgetCreated();
 	UButton* Button = CastChecked<UButton>(LocalGetUMGWidget());
-	Button->OnClicked.AddUniqueDynamic(this, &ThisClass::ServerClick);
+	Button->OnClicked.AddUniqueDynamic(this, &ThisClass::HandleClick);
 }
 
 
 void UUIFrameworkButton::LocalAddChild(FUIFrameworkWidgetId ChildId)
 {
-	if (ChildId == Slot.GetWidgetId())
+	FUIFrameworkWidgetTree* WidgetTree = GetWidgetTree();
+	if (ChildId == Slot.GetWidgetId() && WidgetTree)
 	{
-		check(GetPlayerComponent());
-		if (UUIFrameworkWidget* ChildWidget = GetPlayerComponent()->GetWidgetTree().FindWidgetById(ChildId))
+		if (UUIFrameworkWidget* ChildWidget = WidgetTree->FindWidgetById(ChildId))
 		{
 			UWidget* ChildUMGWidget = ChildWidget->LocalGetUMGWidget();
 			if (ensure(ChildUMGWidget))
@@ -124,15 +114,18 @@ void UUIFrameworkButton::LocalAddChild(FUIFrameworkWidgetId ChildId)
 	}
 }
 
-void UUIFrameworkButton::ServerClick_Implementation()
+void UUIFrameworkButton::HandleClick()
 {
-	if (GetPlayerComponent())
-	{
-		FUIFrameworkClickEventArgument Argument;
-		Argument.PlayerController = GetPlayerComponent()->GetPlayerController();
-		Argument.Sender = this;
-		OnClick.Broadcast(Argument);
-	}
+	// todo the click event should send the userid
+	ServerClick(Cast<APlayerController>(GetOuter()));
+}
+
+void UUIFrameworkButton::ServerClick_Implementation(APlayerController* PlayerController)
+{
+	FUIFrameworkClickEventArgument Argument;
+	Argument.PlayerController = PlayerController;
+	Argument.Sender = this;
+	OnClick.Broadcast(Argument);
 }
 
 void UUIFrameworkButton::OnRep_Slot()
diff --git a/Engine/Plugins/Experimental/UIFramework/Source/Private/Widgets/UIFCanvasBox.cpp b/Engine/Plugins/Experimental/UIFramework/Source/Private/Widgets/UIFCanvasBox.cpp
--- a/Engine/Plugins/Experimental/UIFramework/Source/Private/Widgets/UIFCanvasBox.cpp
+++ b/Engine/Plugins/Experimental/UIFramework/Source/Private/Widgets/UIFCanvasBox.cpp
@@ -101,15 +101,10 @@ void UUIFrameworkCanvasBox::AddWidget(FUIFrameworkCanvasBoxSlot InEntry)
 	{
 		FFrame::KismetExecutionMessage(TEXT("The widget is invalid. It can't be added."), ELogVerbosity::Warning, "InvalidWidgetToAdd");
 	}
-	else if (GetPlayerComponent() && InEntry.AuthorityGetWidget()->GetPlayerComponent() && GetPlayerComponent() != InEntry.AuthorityGetWidget()->GetPlayerComponent())
-	{
-		check(GetPlayerComponent()->GetOwner()->HasAuthority());
-		FFrame::KismetExecutionMessage(TEXT("The widget was created for another player. It can't be removed on this player."), ELogVerbosity::Warning, "InvalidPlayerParentOnRemovedWidget");
-	}
 	else
 	{
 		// Reset the widget to make sure the id is set and it may have been duplicated during the attach
-		InEntry.AuthoritySetWidget(FUIFrameworkModule::AuthorityAttachWidget(GetPlayerComponent(), this, InEntry.AuthorityGetWidget()));
+		InEntry.AuthoritySetWidget(FUIFrameworkModule::AuthorityAttachWidget(this, InEntry.AuthorityGetWidget()));
 		AddEntry(InEntry);
 	}
 }
@@ -120,18 +115,14 @@ void UUIFrameworkCanvasBox::RemoveWidget(UUIFrameworkWidget* Widget)
 	{
 		FFrame::KismetExecutionMessage(TEXT("The widget is invalid. It can't be removed."), ELogVerbosity::Warning, "InvalidWidgetToRemove");
 	}
+	else if (Widget->AuthorityGetParent().IsParentValid() && Widget->AuthorityGetParent().IsWidget() && Widget->AuthorityGetParent().AsWidget() == this)
+	{
+		FFrame::KismetExecutionMessage(TEXT("The widget was created for another widget. It can't be removed on this player."), ELogVerbosity::Warning, "InvalidPlayerParentOnRemovedWidget");
+	}
 	else
 	{
-		if (GetPlayerComponent() && GetPlayerComponent() != Widget->GetPlayerComponent())
-		{
-			check(GetPlayerComponent()->GetOwner()->HasAuthority());
-			FFrame::KismetExecutionMessage(TEXT("The widget was created for another player. It can't be removed on this player."), ELogVerbosity::Warning, "InvalidPlayerParentOnRemovedWidget");
-		}
-		else
-		{
-			FUIFrameworkModule::AuthorityDetachWidgetFromParent(Widget);
-			RemoveEntry(Widget);
-		}
+		FUIFrameworkModule::AuthorityDetachWidgetFromParent(Widget);
+		RemoveEntry(Widget);
 	}
 }
 
@@ -153,25 +144,32 @@ void UUIFrameworkCanvasBox::LocalAddChild(FUIFrameworkWidgetId ChildId)
 	bool bIsAdded = false;
 	if (FUIFrameworkCanvasBoxSlot* CanvasEntry = FindEntry(ChildId))
 	{
-		check(GetPlayerComponent());
-		if (UUIFrameworkWidget* ChildWidget = GetPlayerComponent()->GetWidgetTree().FindWidgetById(ChildId))
+		if (FUIFrameworkWidgetTree* WidgetTree = GetWidgetTree())
 		{
-			UWidget* ChildUMGWidget = ChildWidget->LocalGetUMGWidget();
-			if (ensure(ChildUMGWidget))
+			if (UUIFrameworkWidget* ChildWidget = WidgetTree->FindWidgetById(ChildId))
 			{
-				UCanvasPanelSlot* CanvasSlot = CastChecked<UCanvasPanel>(LocalGetUMGWidget())->AddChildToCanvas(ChildUMGWidget);
-				checkf(CanvasSlot, TEXT("CanvasPanel should be able to receive slot"));
-
-				CanvasEntry->LocalAquireWidget();
+				UWidget* ChildUMGWidget = ChildWidget->LocalGetUMGWidget();
+				if (ensure(ChildUMGWidget))
 				{
-					FAnchorData AnchorData;
-					AnchorData.Anchors = CanvasEntry->Anchors;
-					AnchorData.Offsets = CanvasEntry->Offsets;
-					AnchorData.Alignment = CanvasEntry->Alignment;
-					CanvasSlot->SetLayout(AnchorData);
+					UCanvasPanelSlot* CanvasSlot = CastChecked<UCanvasPanel>(LocalGetUMGWidget())->AddChildToCanvas(ChildUMGWidget);
+					checkf(CanvasSlot, TEXT("CanvasPanel should be able to receive slot"));
+
+					CanvasEntry->LocalAquireWidget();
+					{
+						FAnchorData AnchorData;
+						AnchorData.Anchors = CanvasEntry->Anchors;
+						AnchorData.Offsets = CanvasEntry->Offsets;
+						AnchorData.Alignment = CanvasEntry->Alignment;
+						CanvasSlot->SetLayout(AnchorData);
+					}
+					CanvasSlot->SetZOrder(CanvasEntry->ZOrder);
+					CanvasSlot->SetAutoSize(CanvasEntry->bSizeToContent);
 				}
-				CanvasSlot->SetZOrder(CanvasEntry->ZOrder);
-				CanvasSlot->SetAutoSize(CanvasEntry->bSizeToContent);
+			}
+			else
+			{
+				UE_LOG(LogUIFramework, Log, TEXT("The widget '%" INT64_FMT "' doesn't exist in the WidgetTree."), ChildId.GetKey());
+				Super::LocalAddChild(ChildId);
 			}
 		}
 		else
diff --git a/Engine/Plugins/Experimental/UIFramework/Source/Private/Widgets/UIFSimpleButton.cpp b/Engine/Plugins/Experimental/UIFramework/Source/Private/Widgets/UIFSimpleButton.cpp
--- a/Engine/Plugins/Experimental/UIFramework/Source/Private/Widgets/UIFSimpleButton.cpp
+++ b/Engine/Plugins/Experimental/UIFramework/Source/Private/Widgets/UIFSimpleButton.cpp
@@ -36,7 +36,7 @@ void UUIFrameworkSimpleButton::LocalOnUMGWidgetCreated()
 {
 	UUIFrameworkSimpleButtonUserWidget* UserWidget = CastChecked<UUIFrameworkSimpleButtonUserWidget>(LocalGetUMGWidget());
 	UserWidget->TextBlock->SetText(Text);
-	UserWidget->Button->OnClicked.AddUniqueDynamic(this, &ThisClass::ServerClick);
+	UserWidget->Button->OnClicked.AddUniqueDynamic(this, &ThisClass::HandleClick);
 }
 
 
@@ -56,13 +56,17 @@ void UUIFrameworkSimpleButton::OnRep_Text()
 }
 
 
-void UUIFrameworkSimpleButton::ServerClick_Implementation()
+void UUIFrameworkSimpleButton::HandleClick()
 {
-	if (GetPlayerComponent())
-	{
-		FUIFrameworkClickEventArgument Argument;
-		Argument.PlayerController = GetPlayerComponent()->GetPlayerController();
-		Argument.Sender = this;
-		OnClick.Broadcast(Argument);
-	}
+	// todo the click event should send the userid
+	ServerClick(Cast<APlayerController>(GetOuter()));
+}
+
+
+void UUIFrameworkSimpleButton::ServerClick_Implementation(APlayerController* PlayerController)
+{
+	FUIFrameworkClickEventArgument Argument;
+	Argument.PlayerController = PlayerController;
+	Argument.Sender = this;
+	OnClick.Broadcast(Argument);
 }
diff --git a/Engine/Plugins/Experimental/UIFramework/Source/Private/Widgets/UIFStackBox.cpp b/Engine/Plugins/Experimental/UIFramework/Source/Private/Widgets/UIFStackBox.cpp
--- a/Engine/Plugins/Experimental/UIFramework/Source/Private/Widgets/UIFStackBox.cpp
+++ b/Engine/Plugins/Experimental/UIFramework/Source/Private/Widgets/UIFStackBox.cpp
@@ -107,15 +107,10 @@ void UUIFrameworkStackBox::AddWidget(FUIFrameworkStackBoxSlot InEntry)
 	{
 		FFrame::KismetExecutionMessage(TEXT("The widget is invalid. It can't be added."), ELogVerbosity::Warning, "InvalidWidgetToAdd");
 	}
-	else if (GetPlayerComponent() && InEntry.AuthorityGetWidget()->GetPlayerComponent() && GetPlayerComponent() != InEntry.AuthorityGetWidget()->GetPlayerComponent())
-	{
-		check(GetPlayerComponent()->GetOwner()->HasAuthority());
-		FFrame::KismetExecutionMessage(TEXT("The widget was created for another player. It can't be removed on this player."), ELogVerbosity::Warning, "InvalidPlayerParentOnRemovedWidget");
-	}
 	else
 	{
 		// Reset the widget to make sure the id is set and it may have been duplicated during the attach
-		InEntry.AuthoritySetWidget(FUIFrameworkModule::AuthorityAttachWidget(GetPlayerComponent(), this, InEntry.AuthorityGetWidget()));
+		InEntry.AuthoritySetWidget(FUIFrameworkModule::AuthorityAttachWidget(this, InEntry.AuthorityGetWidget()));
 		AddEntry(InEntry);
 	}
 }
@@ -126,18 +121,14 @@ void UUIFrameworkStackBox::RemoveWidget(UUIFrameworkWidget* Widget)
 	{
 		FFrame::KismetExecutionMessage(TEXT("The widget is invalid. It can't be removed."), ELogVerbosity::Warning, "InvalidWidgetToRemove");
 	}
+	else if (Widget->AuthorityGetParent().IsParentValid() && Widget->AuthorityGetParent().IsWidget() && Widget->AuthorityGetParent().AsWidget() == this)
+	{
+		FFrame::KismetExecutionMessage(TEXT("The widget was created for another widget. It can't be removed on this player."), ELogVerbosity::Warning, "InvalidPlayerParentOnRemovedWidget");
+	}
 	else
 	{
-		if (GetPlayerComponent() && GetPlayerComponent() != Widget->GetPlayerComponent())
-		{
-			check(GetPlayerComponent()->GetOwner()->HasAuthority());
-			FFrame::KismetExecutionMessage(TEXT("The widget was created for another player. It can't be removed on this player."), ELogVerbosity::Warning, "InvalidPlayerParentOnRemovedWidget");
-		}
-		else
-		{
-			FUIFrameworkModule::AuthorityDetachWidgetFromParent(Widget);
-			RemoveEntry(Widget);
-		}
+		FUIFrameworkModule::AuthorityDetachWidgetFromParent(Widget);
+		RemoveEntry(Widget);
 	}
 }
 
@@ -185,46 +176,53 @@ void UUIFrameworkStackBox::LocalAddChild(FUIFrameworkWidgetId ChildId)
 	bool bIsAdded = false;
 	if (FUIFrameworkStackBoxSlot* StackBoxEntry = FindEntry(ChildId))
 	{
-		check(GetPlayerComponent());
-		if (UUIFrameworkWidget* ChildWidget = GetPlayerComponent()->GetWidgetTree().FindWidgetById(ChildId))
+		if (FUIFrameworkWidgetTree* WidgetTree = GetWidgetTree())
 		{
-			UWidget* ChildUMGWidget = ChildWidget->LocalGetUMGWidget();
-			if (ensure(ChildUMGWidget))
+			if (UUIFrameworkWidget* ChildWidget = WidgetTree->FindWidgetById(ChildId))
 			{
-				UStackBox* StackBox = CastChecked<UStackBox>(LocalGetUMGWidget());
-				UPanelSlot* PanelSlot = nullptr;
+				UWidget* ChildUMGWidget = ChildWidget->LocalGetUMGWidget();
+				if (ensure(ChildUMGWidget))
 				{
-					const int32 ChildCount = StackBox->GetChildrenCount();
-					if (StackBoxEntry->Index < ChildCount)
-					{
-						// replace the existing slot
-						StackBox->ReplaceStackBoxChildAt(StackBoxEntry->Index, ChildUMGWidget);
-						PanelSlot = StackBox->GetSlots()[StackBoxEntry->Index];
-					}
-					else if (StackBoxEntry->Index == ChildCount)
+					UStackBox* StackBox = CastChecked<UStackBox>(LocalGetUMGWidget());
+					UPanelSlot* PanelSlot = nullptr;
 					{
-						// add to the end
-						PanelSlot = StackBox->AddChild(ChildUMGWidget);
-					}
-					else
-					{
-						// insert until the widget count is correct
-						for (int32 Index = ChildCount; Index < StackBoxEntry->Index; ++Index)
+						const int32 ChildCount = StackBox->GetChildrenCount();
+						if (StackBoxEntry->Index < ChildCount)
+						{
+							// replace the existing slot
+							StackBox->ReplaceStackBoxChildAt(StackBoxEntry->Index, ChildUMGWidget);
+							PanelSlot = StackBox->GetSlots()[StackBoxEntry->Index];
+						}
+						else if (StackBoxEntry->Index == ChildCount)
+						{
+							// add to the end
+							PanelSlot = StackBox->AddChild(ChildUMGWidget);
+						}
+						else
 						{
-							StackBox->AddChild(NewObject<USpacer>(GetTransientPackage()));
+							// insert until the widget count is correct
+							for (int32 Index = ChildCount; Index < StackBoxEntry->Index; ++Index)
+							{
+								StackBox->AddChild(NewObject<USpacer>(GetTransientPackage()));
+							}
+							PanelSlot = StackBox->AddChild(ChildUMGWidget);
 						}
-						PanelSlot = StackBox->AddChild(ChildUMGWidget);
+						checkf(PanelSlot, TEXT("StackBoxPanel should be able to receive slot"));
 					}
-					checkf(PanelSlot, TEXT("StackBoxPanel should be able to receive slot"));
-				}
 
-				StackBoxEntry->LocalAquireWidget();
+					StackBoxEntry->LocalAquireWidget();
 
-				UStackBoxSlot* StackBoxSlot = CastChecked<UStackBoxSlot>(PanelSlot);
-				StackBoxSlot->SetPadding(StackBoxEntry->Padding);
-				StackBoxSlot->SetSize(StackBoxEntry->Size);
-				StackBoxSlot->SetVerticalAlignment(StackBoxEntry->VerticalAlignment);
-				StackBoxSlot->SetHorizontalAlignment(StackBoxEntry->HorizontalAlignment);
+					UStackBoxSlot* StackBoxSlot = CastChecked<UStackBoxSlot>(PanelSlot);
+					StackBoxSlot->SetPadding(StackBoxEntry->Padding);
+					StackBoxSlot->SetSize(StackBoxEntry->Size);
+					StackBoxSlot->SetVerticalAlignment(StackBoxEntry->VerticalAlignment);
+					StackBoxSlot->SetHorizontalAlignment(StackBoxEntry->HorizontalAlignment);
+				}
+			}
+			else
+			{
+				UE_LOG(LogUIFramework, Log, TEXT("The widget '%" INT64_FMT "' doesn't exist in the WidgetTree."), ChildId.GetKey());
+				Super::LocalAddChild(ChildId);
 			}
 		}
 		else
diff --git a/Engine/Plugins/Experimental/UIFramework/Source/Public/Types/UIFParentWidget.h b/Engine/Plugins/Experimental/UIFramework/Source/Public/Types/UIFParentWidget.h
--- a/Engine/Plugins/Experimental/UIFramework/Source/Public/Types/UIFParentWidget.h
+++ b/Engine/Plugins/Experimental/UIFramework/Source/Public/Types/UIFParentWidget.h
@@ -43,6 +43,23 @@ public:
 	UUIFrameworkWidget* AsWidget() const;
 	UUIFrameworkPlayerComponent* AsPlayerComponent() const;
 
+	bool operator== (const UUIFrameworkWidget* Other) const;
+
+	bool operator!= (const UUIFrameworkWidget* Other) const
+	{
+		return !((*this) == Other);
+	}
+
+	bool operator== (const FUIFrameworkParentWidget& Other) const
+	{
+		return Other.Parent == Parent;
+	}
+
+	bool operator!= (const FUIFrameworkParentWidget& Other) const
+	{
+		return Other.Parent != Parent;
+	}
+
 private:
 	UPROPERTY()
 	TObjectPtr<UObject> Parent = nullptr;
diff --git a/Engine/Plugins/Experimental/UIFramework/Source/Public/Types/UIFWidgetOwner.h b/Engine/Plugins/Experimental/UIFramework/Source/Public/Types/UIFWidgetOwner.h
new file mode 100644
--- /dev/null
+++ b/Engine/Plugins/Experimental/UIFramework/Source/Public/Types/UIFWidgetOwner.h
@@ -0,0 +1,29 @@
+// Copyright Epic Games, Inc. All Rights Reserved.
+
+#pragma once
+
+#include "UObject/ObjectPtr.h"
+
+#include "UIFWidgetOwner.generated.h"
+
+class UGameInstance;
+class APlayerController;
+class UWorld;
+
+/**
+ *
+ */
+USTRUCT()
+struct UIFRAMEWORK_API FUIFrameworkWidgetOwner
+{
+	GENERATED_BODY()
+
+	UPROPERTY()
+	TObjectPtr<APlayerController> PlayerController = nullptr;
+
+	UPROPERTY()
+	TObjectPtr<UGameInstance> GameInstance = nullptr;
+
+	UPROPERTY()
+	TObjectPtr<UWorld> World = nullptr;
+};
\ No newline at end of file
diff --git a/Engine/Plugins/Experimental/UIFramework/Source/Public/Types/UIFWidgetTree.h b/Engine/Plugins/Experimental/UIFramework/Source/Public/Types/UIFWidgetTree.h
--- a/Engine/Plugins/Experimental/UIFramework/Source/Public/Types/UIFWidgetTree.h
+++ b/Engine/Plugins/Experimental/UIFramework/Source/Public/Types/UIFWidgetTree.h
@@ -6,6 +6,7 @@
 #include "UObject/SoftObjectPtr.h"
 #include "Net/Serialization/FastArraySerializer.h"
 #include "Types/UIFWidgetId.h"
+#include "Types/UIFWidgetTreeOwner.h"
 
 #include "UIFWidgetTree.generated.h"
 
@@ -14,6 +15,7 @@
 #endif
 
 struct FReplicationFlags;
+struct FUIFrameworkWidgetTree;
 struct FUIFrameworkWidgetTreeEntry;
 class AActor;
 class FOutBunch;
@@ -24,25 +26,6 @@ class UUIFrameworkWidget;
 class UWidget;
 #endif
 
-class IUIFrameworkWidgetTreeOwner
-{
-public:
-	/** A widget was added to the tree. */
-	virtual void LocalWidgetWasAddedToTree(const FUIFrameworkWidgetTreeEntry& Entry)
-	{
-	}
-
-	/** A widget was removed to the tree. */
-	virtual void LocalWidgetRemovedFromTree(const FUIFrameworkWidgetTreeEntry& Entry)
-	{
-	}
-
-	/** Remove the widget (and the child) from the server. */
-	virtual void LocalRemoveWidgetRootFromTree(const UUIFrameworkWidget* Widget)
-	{
-	}
-};
-
 /**
  *
  */
@@ -68,6 +51,8 @@ struct UIFRAMEWORK_API FUIFrameworkWidgetTreeEntry : public FFastArraySerializer
 
 	UPROPERTY()
 	FUIFrameworkWidgetId ChildId;
+
+	FString GetDebugString();
 };
 
 
@@ -112,7 +97,7 @@ public:
 	 * Remove the widget and all of its children and grand-children from the tree.
 	 * It will clean all the parent relationship from the tree.
 	 */
-	void AuthorityRemoveWidget(UUIFrameworkWidget* Widget);
+	void AuthorityRemoveWidgetAndChildren(UUIFrameworkWidget* Widget);
 
 	/**
 	 * The widget was removed from the client and the Authority is not aware of it.
diff --git a/Engine/Plugins/Experimental/UIFramework/Source/Public/Types/UIFWidgetTreeOwner.h b/Engine/Plugins/Experimental/UIFramework/Source/Public/Types/UIFWidgetTreeOwner.h
new file mode 100644
--- /dev/null
+++ b/Engine/Plugins/Experimental/UIFramework/Source/Public/Types/UIFWidgetTreeOwner.h
@@ -0,0 +1,38 @@
+// Copyright Epic Games, Inc. All Rights Reserved.
+
+#pragma once
+
+#include "Types/UIFWidgetOwner.h"
+
+class APlayerController;
+class UUIFrameworkWidget;
+struct FUIFrameworkWidgetTree;
+struct FUIFrameworkWidgetTreeEntry;
+
+class IUIFrameworkWidgetTreeOwner
+{
+public:
+	/** The widget tree. */
+	virtual FUIFrameworkWidgetTree& GetWidgetTree();
+
+	/** The object that will be used to create the UserWidget. */
+	virtual FUIFrameworkWidgetOwner GetWidgetOwner() const
+	{
+		PURE_VIRTUAL(IUIFrameworkWidgetTreeOwner::GetWidgetOwner, return {};)
+	}
+
+	/** A widget was added to the tree. */
+	virtual void LocalWidgetWasAddedToTree(const FUIFrameworkWidgetTreeEntry& Entry)
+	{
+	}
+
+	/** A widget was removed to the tree. */
+	virtual void LocalWidgetRemovedFromTree(const FUIFrameworkWidgetTreeEntry& Entry)
+	{
+	}
+
+	/** Remove the widget (and the child) from the server. */
+	virtual void LocalRemoveWidgetRootFromTree(const UUIFrameworkWidget* Widget)
+	{
+	}
+};
diff --git a/Engine/Plugins/Experimental/UIFramework/Source/Public/UIFModule.h b/Engine/Plugins/Experimental/UIFramework/Source/Public/UIFModule.h
--- a/Engine/Plugins/Experimental/UIFramework/Source/Public/UIFModule.h
+++ b/Engine/Plugins/Experimental/UIFramework/Source/Public/UIFModule.h
@@ -22,8 +22,8 @@ public:
 	 * If the ReplicationOwner cannot be changed, a duplicate of the widget will be created and the AuthorityOnReattachWidgets will be broadcast.
 	 * This function is recursive if the owner changed.
 	 */
-	static UUIFrameworkWidget* AuthorityAttachWidget(UUIFrameworkPlayerComponent* ReplicationOwner, FUIFrameworkParentWidget Parent, UUIFrameworkWidget* Child);
-	static bool AuthorityCanWidgetBeAttached(UUIFrameworkPlayerComponent* ReplicationOwner, UUIFrameworkWidget* Parent, UUIFrameworkWidget* Child);
+	static UUIFrameworkWidget* AuthorityAttachWidget(FUIFrameworkParentWidget Parent, UUIFrameworkWidget* Child);
+	static bool AuthorityCanWidgetBeAttached(FUIFrameworkParentWidget Parent, UUIFrameworkWidget* Child);
 	/**
 	 * Will remove the widget from the tree and the replication owner.
 	 */
@@ -34,5 +34,7 @@ public:
 	static TSubclassOf<UUIFrameworkPresenter> GetPresenterClass();
 
 private:
-	static void AuthoritySetParentReplicationOwnerRecursive(UUIFrameworkWidget* Widget);
+	static void AuthorityDetachWidgetFromParentInternal(UUIFrameworkWidget* Child, bool bTemporary);
+	//static UUIFrameworkWidget* AuthorityRenameRecursive(UUIFrameworkPlayerComponent* ReplicationOwner, UUIFrameworkWidget* Widget, UObject* NewOuter);
+	//static void AuthoritySetParentReplicationOwnerRecursive(UUIFrameworkWidget* Widget);
 };
\ No newline at end of file
diff --git a/Engine/Plugins/Experimental/UIFramework/Source/Public/UIFPlayerComponent.h b/Engine/Plugins/Experimental/UIFramework/Source/Public/UIFPlayerComponent.h
--- a/Engine/Plugins/Experimental/UIFramework/Source/Public/UIFPlayerComponent.h
+++ b/Engine/Plugins/Experimental/UIFramework/Source/Public/UIFPlayerComponent.h
@@ -125,11 +125,6 @@ public:
 
 	UFUNCTION(BlueprintCallable, BlueprintAuthorityOnly, Category = "UI Framework")
 	void RemoveWidget(UUIFrameworkWidget* Widget);
-
-	FUIFrameworkWidgetTree& GetWidgetTree()
-	{
-		return WidgetTree;
-	}
 	
 	const FUIFrameworkGameLayerSlotList& GetRootList() const
 	{
@@ -160,6 +155,8 @@ public:
 
 	void AuthorityRemoveChild(UUIFrameworkWidget* Widget);
 
+	virtual FUIFrameworkWidgetTree& GetWidgetTree() override;
+	virtual FUIFrameworkWidgetOwner GetWidgetOwner() const override;
 	virtual void LocalWidgetWasAddedToTree(const FUIFrameworkWidgetTreeEntry& Entry) override;
 	virtual void LocalWidgetRemovedFromTree(const FUIFrameworkWidgetTreeEntry& Entry) override;
 	virtual void LocalRemoveWidgetRootFromTree(const UUIFrameworkWidget* Widget) override;
diff --git a/Engine/Plugins/Experimental/UIFramework/Source/Public/UIFWidget.h b/Engine/Plugins/Experimental/UIFramework/Source/Public/UIFWidget.h
--- a/Engine/Plugins/Experimental/UIFramework/Source/Public/UIFWidget.h
+++ b/Engine/Plugins/Experimental/UIFramework/Source/Public/UIFWidget.h
@@ -4,6 +4,7 @@
 
 #include "Components/SlateWrapperTypes.h"
 #include "Components/Widget.h"
+#include "Templates/NonNullPointer.h"
 #include "Types/UIFParentWidget.h"
 #include "Types/UIFWidgetId.h"
 #include "UObject/SoftObjectPtr.h"
@@ -11,6 +12,9 @@
 #include "UIFWidget.generated.h"
 
 class FUIFrameworkModule;
+class UUIFrameworkWidget;
+struct FUIFrameworkWidgetTree;
+class IUIFrameworkWidgetTreeOwner;
 
 /**
  *
@@ -24,6 +28,9 @@ class UUIFrameworkWidgetWrapperInterface : public UInterface
 class IUIFrameworkWidgetWrapperInterface
 {
 	GENERATED_BODY()
+
+public:
+	virtual void ReplaceWidget(UUIFrameworkWidget* OldWidget, UUIFrameworkWidget* NewWidget) {}
 };
 
 /**
@@ -33,7 +40,9 @@ UCLASS(Abstract, BlueprintType)
 class UIFRAMEWORK_API UUIFrameworkWidget : public UObject
 {
 	GENERATED_BODY()
+
 	friend FUIFrameworkModule;
+	friend FUIFrameworkWidgetTree;
 
 private:
 	UPROPERTY(BlueprintReadWrite, ReplicatedUsing = "OnRep_IsEnabled", Getter = "IsEnabled", Setter="SetEnabled", Category = "UI Framework", meta = (AllowPrivateAccess = "true"))
@@ -67,11 +76,13 @@ public:
 		return Id;
 	}
 
-	UUIFrameworkPlayerComponent* GetPlayerComponent() const
+	IUIFrameworkWidgetTreeOwner* GetWidgetTreeOwner() const
 	{
-		return OwnerPlayerComponent;
+		return WidgetTreeOwner;
 	}
 
+	FUIFrameworkWidgetTree* GetWidgetTree() const;
+
 	TSoftClassPtr<UWidget> GetUMGWidgetClass() const
 	{
 		return WidgetClass;
@@ -93,7 +104,7 @@ public:
 		return LocalUMGWidget;
 	}
 
-	void LocalCreateUMGWidget(UUIFrameworkPlayerComponent* Owner);
+	void LocalCreateUMGWidget(TNonNullPtr<IUIFrameworkWidgetTreeOwner> Owner);
 	virtual void LocalAddChild(FUIFrameworkWidgetId ChildId);
 	void LocalDestroyUMGWidget();
 
@@ -128,12 +139,11 @@ private:
 	FUIFrameworkWidgetId Id = FUIFrameworkWidgetId::MakeNew();
 
 	//~ Authority
-	UPROPERTY(Transient, DuplicateTransient)
+	UPROPERTY(Transient)
 	TScriptInterface<IUIFrameworkWidgetWrapperInterface> Wrapper;
 
 	//~ Authority and Local
-	UPROPERTY(Transient)
-	TObjectPtr<UUIFrameworkPlayerComponent> OwnerPlayerComponent = nullptr;
+	IUIFrameworkWidgetTreeOwner* WidgetTreeOwner = nullptr;
 
 	//~ AuthorityOnly
 	UPROPERTY(Transient)
diff --git a/Engine/Plugins/Experimental/UIFramework/Source/Public/Widgets/UIFButton.h b/Engine/Plugins/Experimental/UIFramework/Source/Public/Widgets/UIFButton.h
--- a/Engine/Plugins/Experimental/UIFramework/Source/Public/Widgets/UIFButton.h
+++ b/Engine/Plugins/Experimental/UIFramework/Source/Public/Widgets/UIFButton.h
@@ -36,8 +36,11 @@ protected:
 	virtual void LocalOnUMGWidgetCreated() override;
 
 private:
+	UFUNCTION()
+	void HandleClick();
+
 	UFUNCTION(Server, Reliable)
-	void ServerClick();
+	void ServerClick(APlayerController* PlayerController);
 
 	UFUNCTION()
 	void OnRep_Slot();
diff --git a/Engine/Plugins/Experimental/UIFramework/Source/Public/Widgets/UIFSimpleButton.h b/Engine/Plugins/Experimental/UIFramework/Source/Public/Widgets/UIFSimpleButton.h
--- a/Engine/Plugins/Experimental/UIFramework/Source/Public/Widgets/UIFSimpleButton.h
+++ b/Engine/Plugins/Experimental/UIFramework/Source/Public/Widgets/UIFSimpleButton.h
@@ -33,8 +33,11 @@ private:
 	UFUNCTION()
 	void OnRep_Text();
 
+	UFUNCTION()
+	void HandleClick();
+
 	UFUNCTION(Server, Reliable)
-	void ServerClick();
+	void ServerClick(APlayerController* PlayerController);
 
 private:
 	UPROPERTY(ReplicatedUsing=OnRep_Text)
