Fix CIS 241525
#rnx

diff --git a/Engine/Plugins/Experimental/UIFramework/Source/Private/Types/UIFWidgetTreeOwner.cpp b/Engine/Plugins/Experimental/UIFramework/Source/Private/Types/UIFWidgetTreeOwner.cpp
--- a/Engine/Plugins/Experimental/UIFramework/Source/Private/Types/UIFWidgetTreeOwner.cpp
+++ b/Engine/Plugins/Experimental/UIFramework/Source/Private/Types/UIFWidgetTreeOwner.cpp
@@ -6,5 +6,6 @@
 
 FUIFrameworkWidgetTree& IUIFrameworkWidgetTreeOwner::GetWidgetTree()
 {
-	PURE_VIRTUAL(IUIFrameworkWidgetTreeOwner::GetWidgetTree, static FUIFrameworkWidgetTree Tmp; return Tmp;)
+	static FUIFrameworkWidgetTree Tmp;
+	PURE_VIRTUAL(IUIFrameworkWidgetTreeOwner::GetWidgetTree, return Tmp;)
 }
\ No newline at end of file
