Fix CIS 241525
#rnx

diff --git a/Engine/Plugins/Experimental/UIFramework/Source/Private/Types/UIFWidgetTreeOwner.cpp b/Engine/Plugins/Experimental/UIFramework/Source/Private/Types/UIFWidgetTreeOwner.cpp
--- a/Engine/Plugins/Experimental/UIFramework/Source/Private/Types/UIFWidgetTreeOwner.cpp
+++ b/Engine/Plugins/Experimental/UIFramework/Source/Private/Types/UIFWidgetTreeOwner.cpp
@@ -2,10 +2,3 @@
 
 #include "Types/UIFWidgetTreeOwner.h"
 #include "Types/UIFWidgetTree.h"
-
-
-FUIFrameworkWidgetTree& IUIFrameworkWidgetTreeOwner::GetWidgetTree()
-{
-	static FUIFrameworkWidgetTree Tmp;
-	PURE_VIRTUAL(IUIFrameworkWidgetTreeOwner::GetWidgetTree, return Tmp;)
-}
\ No newline at end of file
diff --git a/Engine/Plugins/Experimental/UIFramework/Source/Public/Types/UIFWidgetTreeOwner.h b/Engine/Plugins/Experimental/UIFramework/Source/Public/Types/UIFWidgetTreeOwner.h
--- a/Engine/Plugins/Experimental/UIFramework/Source/Public/Types/UIFWidgetTreeOwner.h
+++ b/Engine/Plugins/Experimental/UIFramework/Source/Public/Types/UIFWidgetTreeOwner.h
@@ -13,13 +13,10 @@ class IUIFrameworkWidgetTreeOwner
 {
 public:
 	/** The widget tree. */
-	virtual FUIFrameworkWidgetTree& GetWidgetTree();
+	virtual FUIFrameworkWidgetTree& GetWidgetTree() = 0;
 
 	/** The object that will be used to create the UserWidget. */
-	virtual FUIFrameworkWidgetOwner GetWidgetOwner() const
-	{
-		PURE_VIRTUAL(IUIFrameworkWidgetTreeOwner::GetWidgetOwner, return {};)
-	}
+	virtual FUIFrameworkWidgetOwner GetWidgetOwner() const = 0;
 
 	/** A widget was added to the tree. */
 	virtual void LocalWidgetWasAddedToTree(const FUIFrameworkWidgetTreeEntry& Entry)
