Add text justification to UI.

diff --git a/Engine/Plugins/Experimental/UIFramework/Source/Private/Widgets/UIFTextBlock.cpp b/Engine/Plugins/Experimental/UIFramework/Source/Private/Widgets/UIFTextBlock.cpp
--- a/Engine/Plugins/Experimental/UIFramework/Source/Private/Widgets/UIFTextBlock.cpp
+++ b/Engine/Plugins/Experimental/UIFramework/Source/Private/Widgets/UIFTextBlock.cpp
@@ -26,21 +26,30 @@ void UUIFrameworkTextBlock::GetLifetimeReplicatedProps(TArray<FLifetimeProperty>
 	FDoRepLifetimeParams Params;
 	Params.bIsPushBased = true;
 	DOREPLIFETIME_WITH_PARAMS_FAST(ThisClass, Text, Params);
+	DOREPLIFETIME_WITH_PARAMS_FAST(ThisClass, Justification, Params);
 }
 
-
 void UUIFrameworkTextBlock::LocalOnUMGWidgetCreated()
 {
-	CastChecked<UTextBlock>(LocalGetUMGWidget())->SetText(Text);
+	UTextBlock* TextBlock = CastChecked<UTextBlock>(LocalGetUMGWidget());
+	TextBlock->SetText(Text);
+	TextBlock->SetJustification(Justification);
 }
 
-
 void UUIFrameworkTextBlock::SetText(FText InText)
 {
 	Text = InText;
 	MARK_PROPERTY_DIRTY_FROM_NAME(ThisClass, Text, this);
 }
 
+void UUIFrameworkTextBlock::SetJustification(ETextJustify::Type InJustification)
+{
+	if (Justification != InJustification)
+	{
+		Justification = InJustification;
+		MARK_PROPERTY_DIRTY_FROM_NAME(ThisClass, Justification, this);
+	}
+}
 
 void UUIFrameworkTextBlock::OnRep_Text()
 {
@@ -49,3 +58,11 @@ void UUIFrameworkTextBlock::OnRep_Text()
 		CastChecked<UTextBlock>(LocalGetUMGWidget())->SetText(Text);
 	}
 }
+
+void UUIFrameworkTextBlock::OnRep_Justification()
+{
+	if (LocalGetUMGWidget())
+	{
+		CastChecked<UTextBlock>(LocalGetUMGWidget())->SetJustification(Justification);
+	}
+}
diff --git a/Engine/Plugins/Experimental/UIFramework/Source/Public/Widgets/UIFTextBlock.h b/Engine/Plugins/Experimental/UIFramework/Source/Public/Widgets/UIFTextBlock.h
--- a/Engine/Plugins/Experimental/UIFramework/Source/Public/Widgets/UIFTextBlock.h
+++ b/Engine/Plugins/Experimental/UIFramework/Source/Public/Widgets/UIFTextBlock.h
@@ -2,6 +2,7 @@
 
 #pragma once
 
+#include "Framework/Text/TextLayout.h"
 #include "UIFWidget.h"
 
 #include "UIFTextBlock.generated.h"
@@ -27,6 +28,15 @@ public:
 	{
 		return Text;
 	}
+	
+	UFUNCTION(BlueprintCallable, BlueprintAuthorityOnly, Category = "UI Framework")
+	void SetJustification(ETextJustify::Type Justification);
+
+	UFUNCTION(BlueprintCallable, Category = "UI Framework")
+	ETextJustify::Type GetJustification() const
+	{
+		return Justification;
+	}
 
 	virtual void LocalOnUMGWidgetCreated() override;
 
@@ -34,7 +44,13 @@ private:
 	UFUNCTION()
 	void OnRep_Text();
 
+	UFUNCTION()
+	void OnRep_Justification();
+
 private:
 	UPROPERTY(/*ExposeOnSpawn, */ReplicatedUsing=OnRep_Text)
 	FText Text;
+
+	UPROPERTY(ReplicatedUsing=OnRep_Justification)
+	TEnumAsByte<ETextJustify::Type> Justification;
 };
