Force net update on widget owner when a replicated property changes that requires immediate UI update

diff --git a/Engine/Plugins/Experimental/UIFramework/Source/Private/UIFWidget.cpp b/Engine/Plugins/Experimental/UIFramework/Source/Private/UIFWidget.cpp
--- a/Engine/Plugins/Experimental/UIFramework/Source/Private/UIFWidget.cpp
+++ b/Engine/Plugins/Experimental/UIFramework/Source/Private/UIFWidget.cpp
@@ -20,6 +20,14 @@
 #include UE_INLINE_GENERATED_CPP_BY_NAME(UIFWidget)
 
 
+void UUIFrameworkWidget::ForceNetUpdate()
+{
+	if (AActor* OwnerActor = Cast<AActor>(GetOuter()))
+	{
+		OwnerActor->ForceNetUpdate();
+	}
+}
+
 /**
  *
  */
@@ -150,6 +158,7 @@ void UUIFrameworkWidget::SetVisibility(ESlateVisibility InVisibility)
 	{
 		Visibility = InVisibility;
 		MARK_PROPERTY_DIRTY_FROM_NAME(ThisClass, Visibility, this);
+		ForceNetUpdate();
 	}
 }
 
@@ -164,6 +173,7 @@ void UUIFrameworkWidget::SetEnabled(bool bInIsEnabled)
 	{
 		bIsEnabled = bInIsEnabled;
 		MARK_PROPERTY_DIRTY_FROM_NAME(ThisClass, bIsEnabled, this);
+		ForceNetUpdate();
 	}
 }
 
diff --git a/Engine/Plugins/Experimental/UIFramework/Source/Private/Widgets/UIFButton.cpp b/Engine/Plugins/Experimental/UIFramework/Source/Private/Widgets/UIFButton.cpp
--- a/Engine/Plugins/Experimental/UIFramework/Source/Private/Widgets/UIFButton.cpp
+++ b/Engine/Plugins/Experimental/UIFramework/Source/Private/Widgets/UIFButton.cpp
@@ -57,6 +57,7 @@ void UUIFrameworkButton::SetContent(FUIFrameworkSimpleSlot InEntry)
 	}
 
 	MARK_PROPERTY_DIRTY_FROM_NAME(ThisClass, Slot, this);
+	ForceNetUpdate();
 }
 
 
@@ -76,6 +77,7 @@ void UUIFrameworkButton::AuthorityRemoveChild(UUIFrameworkWidget* Widget)
 
 	Slot.AuthoritySetWidget(nullptr);;
 	MARK_PROPERTY_DIRTY_FROM_NAME(ThisClass, Slot, this);
+	ForceNetUpdate();
 }
 
 void UUIFrameworkButton::LocalOnUMGWidgetCreated()
diff --git a/Engine/Plugins/Experimental/UIFramework/Source/Private/Widgets/UIFSimpleButton.cpp b/Engine/Plugins/Experimental/UIFramework/Source/Private/Widgets/UIFSimpleButton.cpp
--- a/Engine/Plugins/Experimental/UIFramework/Source/Private/Widgets/UIFSimpleButton.cpp
+++ b/Engine/Plugins/Experimental/UIFramework/Source/Private/Widgets/UIFSimpleButton.cpp
@@ -44,6 +44,7 @@ void UUIFrameworkSimpleButton::SetText(FText InText)
 {
 	Text = InText;
 	MARK_PROPERTY_DIRTY_FROM_NAME(ThisClass, Text, this);
+	ForceNetUpdate();
 }
 
 
diff --git a/Engine/Plugins/Experimental/UIFramework/Source/Private/Widgets/UIFStackBox.cpp b/Engine/Plugins/Experimental/UIFramework/Source/Private/Widgets/UIFStackBox.cpp
--- a/Engine/Plugins/Experimental/UIFramework/Source/Private/Widgets/UIFStackBox.cpp
+++ b/Engine/Plugins/Experimental/UIFramework/Source/Private/Widgets/UIFStackBox.cpp
@@ -147,6 +147,7 @@ void UUIFrameworkStackBox::SetOrientation(EOrientation Value)
 			StackBox->SetOrientation(Orientation);
 		}
 		MARK_PROPERTY_DIRTY_FROM_NAME(ThisClass, Orientation, this);
+		ForceNetUpdate();
 	}
 }
 
diff --git a/Engine/Plugins/Experimental/UIFramework/Source/Private/Widgets/UIFTextBlock.cpp b/Engine/Plugins/Experimental/UIFramework/Source/Private/Widgets/UIFTextBlock.cpp
--- a/Engine/Plugins/Experimental/UIFramework/Source/Private/Widgets/UIFTextBlock.cpp
+++ b/Engine/Plugins/Experimental/UIFramework/Source/Private/Widgets/UIFTextBlock.cpp
@@ -40,6 +40,7 @@ void UUIFrameworkTextBlock::SetText(FText InText)
 {
 	Text = InText;
 	MARK_PROPERTY_DIRTY_FROM_NAME(ThisClass, Text, this);
+	ForceNetUpdate();
 }
 
 void UUIFrameworkTextBlock::SetJustification(ETextJustify::Type InJustification)
@@ -48,6 +49,7 @@ void UUIFrameworkTextBlock::SetJustification(ETextJustify::Type InJustification)
 	{
 		Justification = InJustification;
 		MARK_PROPERTY_DIRTY_FROM_NAME(ThisClass, Justification, this);
+		ForceNetUpdate();
 	}
 }
 
diff --git a/Engine/Plugins/Experimental/UIFramework/Source/Public/UIFWidget.h b/Engine/Plugins/Experimental/UIFramework/Source/Public/UIFWidget.h
--- a/Engine/Plugins/Experimental/UIFramework/Source/Public/UIFWidget.h
+++ b/Engine/Plugins/Experimental/UIFramework/Source/Public/UIFWidget.h
@@ -123,6 +123,8 @@ protected:
 	{
 	}
 
+	void ForceNetUpdate();
+
 private:
 	UFUNCTION()
 	void OnRep_IsEnabled();
