UIFramework: User can add any userwidget asset.

diff --git a/Engine/Plugins/Experimental/UIFramework/Source/Private/UIFPlayerComponent.cpp b/Engine/Plugins/Experimental/UIFramework/Source/Private/UIFPlayerComponent.cpp
--- a/Engine/Plugins/Experimental/UIFramework/Source/Private/UIFPlayerComponent.cpp
+++ b/Engine/Plugins/Experimental/UIFramework/Source/Private/UIFPlayerComponent.cpp
@@ -272,10 +272,10 @@ void UUIFrameworkPlayerComponent::LocalWidgetWasAddedToTree(const FUIFrameworkWi
 	APlayerController* LocalOwner = GetPlayerController();
 	check(!LocalOwner->HasAuthority());
 
-	if (Entry.Child)
+	if (Entry.Child && Entry.Child->LocalIsReplicationReady())
 	{
 		TSoftClassPtr<UWidget> WidgetClass = Entry.Child->GetUMGWidgetClass();
-		if (WidgetClass.Get() != nullptr)
+		if (WidgetClass.IsValid())
 		{
 			if (Entry.IsParentValid() && Entry.IsChildValid())
 			{
diff --git a/Engine/Plugins/Experimental/UIFramework/Source/Private/UIFWidget.cpp b/Engine/Plugins/Experimental/UIFramework/Source/Private/UIFWidget.cpp
--- a/Engine/Plugins/Experimental/UIFramework/Source/Private/UIFWidget.cpp
+++ b/Engine/Plugins/Experimental/UIFramework/Source/Private/UIFWidget.cpp
@@ -76,6 +76,9 @@ void UUIFrameworkWidget::GetLifetimeReplicatedProps(TArray<FLifetimeProperty>& O
 	DOREPLIFETIME_WITH_PARAMS_FAST(ThisClass, Id, Params);
 	DOREPLIFETIME_WITH_PARAMS_FAST(ThisClass, bIsEnabled, Params);
 	DOREPLIFETIME_WITH_PARAMS_FAST(ThisClass, Visibility, Params);
+
+	Params.Condition = COND_InitialOnly;
+	DOREPLIFETIME_WITH_PARAMS_FAST(UUIFrameworkWidget, WidgetClass, Params);
 }
 
 FUIFrameworkWidgetTree* UUIFrameworkWidget::GetWidgetTree() const
diff --git a/Engine/Plugins/Experimental/UIFramework/Source/Private/Widgets/UIFUserWidget.cpp b/Engine/Plugins/Experimental/UIFramework/Source/Private/Widgets/UIFUserWidget.cpp
new file mode 100644
--- /dev/null
+++ b/Engine/Plugins/Experimental/UIFramework/Source/Private/Widgets/UIFUserWidget.cpp
@@ -0,0 +1,205 @@
+// Copyright Epic Games, Inc. All Rights Reserved.
+
+#include "Widgets/UIFUserWidget.h"
+#include "Types/UIFWidgetTree.h"
+#include "UIFLog.h"
+#include "UIFModule.h"
+#include "UIFPlayerComponent.h"
+
+#include "Blueprint/UserWidget.h"
+
+#include "Engine/ActorChannel.h"
+#include "Engine/Engine.h"
+#include "Engine/NetDriver.h"
+#include "GameFramework/Actor.h"
+#include "GameFramework/PlayerController.h"
+#include "Net/UnrealNetwork.h"
+
+#include UE_INLINE_GENERATED_CPP_BY_NAME(UIFUserWidget)
+
+
+/**
+ *
+ */
+void FUIFrameworkUserWidgetSlotList::PostReplicatedChange(const TArrayView<int32>& ChangedIndices, int32 FinalSize)
+{
+	check(Owner);
+	for (int32 Index : ChangedIndices)
+	{
+		FUIFrameworkUserWidgetSlot& Slot = Slots[Index];
+		if (Slot.LocalIsAquiredWidgetValid())
+		{
+			// Remove and add the widget again... 
+			// that may not work if they are on top of each other... The order may matter if the zorder is the same :( 
+			Owner->LocalAddChild(Slot.GetWidgetId());
+		}
+		//else it will be remove and the new widget will be added by the WidgetTree replication.
+	}
+}
+
+bool FUIFrameworkUserWidgetSlotList::NetDeltaSerialize(FNetDeltaSerializeInfo& DeltaParms)
+{
+	return FFastArraySerializer::FastArrayDeltaSerialize<FUIFrameworkUserWidgetSlot, FUIFrameworkUserWidgetSlotList>(Slots, DeltaParms, *this);
+}
+
+void FUIFrameworkUserWidgetSlotList::AddEntry(FUIFrameworkUserWidgetSlot Entry)
+{
+	FUIFrameworkUserWidgetSlot& NewEntry = Slots.Add_GetRef(MoveTemp(Entry));
+	MarkItemDirty(NewEntry);
+}
+
+bool FUIFrameworkUserWidgetSlotList::RemoveEntry(UUIFrameworkWidget* Widget)
+{
+	check(Widget);
+	const int32 Index = Slots.IndexOfByPredicate([Widget](const FUIFrameworkUserWidgetSlot& Entry) { return Entry.AuthorityGetWidget() == Widget; });
+	if (Index != INDEX_NONE)
+	{
+		Slots.RemoveAt(Index);
+		MarkArrayDirty();
+	}
+	return Index != INDEX_NONE;
+}
+
+FUIFrameworkUserWidgetSlot* FUIFrameworkUserWidgetSlotList::FindEntry(FUIFrameworkWidgetId WidgetId)
+{
+	return Slots.FindByPredicate([WidgetId](const FUIFrameworkUserWidgetSlot& Entry) { return Entry.GetWidgetId() == WidgetId; });
+}
+
+void FUIFrameworkUserWidgetSlotList::ForEachChildren(const TFunctionRef<void(UUIFrameworkWidget*)>& Func)
+{
+	for (FUIFrameworkUserWidgetSlot& Slot : Slots)
+	{
+		if (UUIFrameworkWidget* ChildWidget = Slot.AuthorityGetWidget())
+		{
+			Func(ChildWidget);
+		}
+	}
+}
+
+/**
+ *
+ */
+UUIFrameworkUserWidget::UUIFrameworkUserWidget()
+	: ReplicatedSlotList(this)
+{
+}
+
+void UUIFrameworkUserWidget::GetLifetimeReplicatedProps(TArray<FLifetimeProperty>& OutLifetimeProps) const
+{
+	Super::GetLifetimeReplicatedProps(OutLifetimeProps);
+
+	FDoRepLifetimeParams Params;
+	Params.bIsPushBased = true;
+	DOREPLIFETIME_WITH_PARAMS_FAST(ThisClass, ReplicatedSlotList, Params);
+}
+
+void UUIFrameworkUserWidget::LocalOnUMGWidgetCreated()
+{
+	Super::LocalOnUMGWidgetCreated();
+	//todo add viewmodel
+}
+
+bool UUIFrameworkUserWidget::LocalIsReplicationReady() const
+{
+	return Super::LocalIsReplicationReady() && !WidgetClass.IsNull();
+}
+
+void UUIFrameworkUserWidget::SetWidgetClass(TSoftClassPtr<UWidget> InWidgetClass)
+{
+	WidgetClass = InWidgetClass;
+	MARK_PROPERTY_DIRTY_FROM_NAME(UUIFrameworkWidget, WidgetClass, this);
+}
+
+void UUIFrameworkUserWidget::SetNamedSlot(FName SlotName, UUIFrameworkWidget* Widget)
+{
+	if (Widget == nullptr || SlotName.IsNone())
+	{
+		FFrame::KismetExecutionMessage(TEXT("The widget is invalid. It can't be added."), ELogVerbosity::Warning, "InvalidWidgetToAdd");
+	}
+	else
+	{
+		// Reset the widget to make sure the id is set and it may have been duplicated during the attach
+		FUIFrameworkUserWidgetSlot Entry;
+		Entry.SlotName = SlotName;
+		Entry.AuthoritySetWidget(Widget);
+		Entry.AuthoritySetWidget(FUIFrameworkModule::AuthorityAttachWidget(this, Entry.AuthorityGetWidget()));
+		AddEntry(Entry);
+	}
+}
+
+
+void UUIFrameworkUserWidget::AuthorityForEachChildren(const TFunctionRef<void(UUIFrameworkWidget*)>& Func)
+{
+	Super::AuthorityForEachChildren(Func);
+
+	ReplicatedSlotList.ForEachChildren(Func);
+}
+
+void UUIFrameworkUserWidget::AuthorityRemoveChild(UUIFrameworkWidget* Widget)
+{
+	Super::AuthorityRemoveChild(Widget);
+	RemoveEntry(Widget);
+}
+
+void UUIFrameworkUserWidget::LocalAddChild(FUIFrameworkWidgetId ChildId)
+{
+	bool bIsAdded = false;
+	if (FUIFrameworkUserWidgetSlot* NamedSlotEntry = FindEntry(ChildId))
+	{
+		if (FUIFrameworkWidgetTree* WidgetTree = GetWidgetTree())
+		{
+			if (UUIFrameworkWidget* ChildWidget = WidgetTree->FindWidgetById(ChildId))
+			{
+				UWidget* ChildUMGWidget = ChildWidget->LocalGetUMGWidget();
+				if (ensure(ChildUMGWidget))
+				{
+					NamedSlotEntry->LocalAquireWidget();
+					if (UUserWidget* LocalUMGUserWidget = Cast<UUserWidget>(LocalGetUMGWidget()))
+					{
+						LocalUMGUserWidget->SetContentForSlot(NamedSlotEntry->SlotName, ChildUMGWidget);
+					}
+					else
+					{
+						UE_LOG(LogUIFramework, Log, TEXT("Can't set the NamedSlot on widget '%s' because it is not a UserWidget."), *ChildUMGWidget->GetName());
+						Super::LocalAddChild(ChildId);
+					}
+				}
+				else
+				{
+					UE_LOG(LogUIFramework, Error, TEXT("The widget '%" INT64_FMT "' is invalid."), ChildId.GetKey());
+					Super::LocalAddChild(ChildId);
+				}
+			}
+			else
+			{
+				UE_LOG(LogUIFramework, Log, TEXT("The widget '%" INT64_FMT "' doesn't exist in the WidgetTree."), ChildId.GetKey());
+				Super::LocalAddChild(ChildId);
+			}
+		}
+		else
+		{
+			UE_LOG(LogUIFramework, Log, TEXT("The widget '%" INT64_FMT "' doesn't exist in the WidgetTree."), ChildId.GetKey());
+			Super::LocalAddChild(ChildId);
+		}
+	}
+	else
+	{
+		UE_LOG(LogUIFramework, Verbose, TEXT("The widget '%" INT64_FMT "' was not found in the Canvas Slots."), ChildId.GetKey());
+		Super::LocalAddChild(ChildId);
+	}
+}
+
+void UUIFrameworkUserWidget::AddEntry(FUIFrameworkUserWidgetSlot Entry)
+{
+	ReplicatedSlotList.AddEntry(Entry);
+}
+
+bool UUIFrameworkUserWidget::RemoveEntry(UUIFrameworkWidget* Widget)
+{
+	return ReplicatedSlotList.RemoveEntry(Widget);
+}
+
+FUIFrameworkUserWidgetSlot* UUIFrameworkUserWidget::FindEntry(FUIFrameworkWidgetId WidgetId)
+{
+	return ReplicatedSlotList.FindEntry(WidgetId);
+}
diff --git a/Engine/Plugins/Experimental/UIFramework/Source/Public/UIFWidget.h b/Engine/Plugins/Experimental/UIFramework/Source/Public/UIFWidget.h
--- a/Engine/Plugins/Experimental/UIFramework/Source/Public/UIFWidget.h
+++ b/Engine/Plugins/Experimental/UIFramework/Source/Public/UIFWidget.h
@@ -99,6 +99,11 @@ public:
 	}
 
 	//~ Local functions
+	virtual bool LocalIsReplicationReady() const
+	{
+		return true;
+	}
+
 	UWidget* LocalGetUMGWidget() const
 	{
 		return LocalUMGWidget;
@@ -132,7 +137,7 @@ private:
 	void OnRep_Visibility();
 
 protected:
-	UPROPERTY(BlueprintReadOnly, EditDefaultsOnly, Category = "UI Framework")
+	UPROPERTY(BlueprintReadOnly, Replicated, EditDefaultsOnly, Category = "UI Framework")
 	TSoftClassPtr<UWidget> WidgetClass; // todo: make this private and use a constructor argument
 
 private:
diff --git a/Engine/Plugins/Experimental/UIFramework/Source/Public/Widgets/UIFUserWidget.h b/Engine/Plugins/Experimental/UIFramework/Source/Public/Widgets/UIFUserWidget.h
new file mode 100644
--- /dev/null
+++ b/Engine/Plugins/Experimental/UIFramework/Source/Public/Widgets/UIFUserWidget.h
@@ -0,0 +1,101 @@
+// Copyright Epic Games, Inc. All Rights Reserved.
+
+#pragma once
+
+#include "UIFWidget.h"
+#include "Types/UIFSlotBase.h"
+
+#include "UIFUserWidget.generated.h"
+
+class UUserWidget;
+class UUIFrameworkUserWidget;
+
+/**
+ *
+ */
+USTRUCT(BlueprintType)
+struct FUIFrameworkUserWidgetSlot : public FUIFrameworkSlotBase
+{
+	GENERATED_BODY()
+
+	/** The name of the NamedSlot */
+	UPROPERTY(EditAnywhere, BlueprintReadWrite, Category = "UI Framework")
+	FName SlotName;
+};
+
+/**
+ *
+ */
+USTRUCT()
+struct UIFRAMEWORK_API FUIFrameworkUserWidgetSlotList : public FFastArraySerializer
+{
+	GENERATED_BODY()
+
+	FUIFrameworkUserWidgetSlotList() = default;
+	FUIFrameworkUserWidgetSlotList(UUIFrameworkUserWidget* InOwner)
+		: Owner(InOwner)
+	{}
+
+public:
+	//~ Begin of FFastArraySerializer
+	void PostReplicatedChange(const TArrayView<int32>& ChangedIndices, int32 FinalSize);
+	//~ End of FFastArraySerializer
+
+	bool NetDeltaSerialize(FNetDeltaSerializeInfo& DeltaParms);
+
+	void AddEntry(FUIFrameworkUserWidgetSlot Entry);
+	bool RemoveEntry(UUIFrameworkWidget* Widget);
+	FUIFrameworkUserWidgetSlot* FindEntry(FUIFrameworkWidgetId WidgetId);
+	void ForEachChildren(const TFunctionRef<void(UUIFrameworkWidget*)>& Func);
+
+private:
+	UPROPERTY()
+	TArray<FUIFrameworkUserWidgetSlot> Slots;
+
+	UPROPERTY(NotReplicated, Transient)
+	TObjectPtr<UUIFrameworkUserWidget> Owner;
+};
+
+
+template<>
+struct TStructOpsTypeTraits<FUIFrameworkUserWidgetSlotList> : public TStructOpsTypeTraitsBase2<FUIFrameworkUserWidgetSlotList>
+{
+	enum { WithNetDeltaSerializer = true };
+};
+
+/**
+ *
+ */
+UCLASS(DisplayName = "UserWidget UIFramework")
+class UIFRAMEWORK_API UUIFrameworkUserWidget : public UUIFrameworkWidget
+{
+	GENERATED_BODY()
+
+public:
+	UUIFrameworkUserWidget();
+
+public:
+
+	UFUNCTION(BlueprintCallable, BlueprintAuthorityOnly, Category = "UI Framework")
+	void SetWidgetClass(TSoftClassPtr<UWidget> Value);
+
+	UFUNCTION(BlueprintCallable, BlueprintAuthorityOnly, Category = "UI Framework")
+	void SetNamedSlot(FName SlotName, UUIFrameworkWidget* Widget);
+
+public:
+	void LocalOnUMGWidgetCreated() override;
+	virtual bool LocalIsReplicationReady() const;
+
+	virtual void AuthorityForEachChildren(const TFunctionRef<void(UUIFrameworkWidget*)>& Func) override;
+	virtual void AuthorityRemoveChild(UUIFrameworkWidget* Widget) override;
+	virtual void LocalAddChild(FUIFrameworkWidgetId ChildId) override;
+
+private:
+	void AddEntry(FUIFrameworkUserWidgetSlot Entry);
+	bool RemoveEntry(UUIFrameworkWidget* Widget);
+	FUIFrameworkUserWidgetSlot* FindEntry(FUIFrameworkWidgetId WidgetId);
+
+private:
+	UPROPERTY(Replicated)
+	FUIFrameworkUserWidgetSlotList ReplicatedSlotList;
+};
