UIFramwork: Use MVVM for SimpleButton

diff --git a/Engine/Plugins/Experimental/UIFramework/Source/Private/Widgets/UIFSimpleButton.cpp b/Engine/Plugins/Experimental/UIFramework/Source/Private/Widgets/UIFSimpleButton.cpp
--- a/Engine/Plugins/Experimental/UIFramework/Source/Private/Widgets/UIFSimpleButton.cpp
+++ b/Engine/Plugins/Experimental/UIFramework/Source/Private/Widgets/UIFSimpleButton.cpp
@@ -6,9 +6,9 @@
 #include "Blueprint/WidgetTree.h"
 #include "Components/Button.h"
 #include "Components/TextBlock.h"
-#include "Components/UIFSimpleButtonUserWidget.h"
-
+#include "MVVMSubsystem.h"
 #include "Net/UnrealNetwork.h"
+#include "View/MVVMView.h"
 
 #include UE_INLINE_GENERATED_CPP_BY_NAME(UIFSimpleButton)
 
@@ -34,9 +34,14 @@ void UUIFrameworkSimpleButton::GetLifetimeReplicatedProps(TArray<FLifetimeProper
 
 void UUIFrameworkSimpleButton::LocalOnUMGWidgetCreated()
 {
-	UUIFrameworkSimpleButtonUserWidget* UserWidget = CastChecked<UUIFrameworkSimpleButtonUserWidget>(LocalGetUMGWidget());
-	UserWidget->TextBlock->SetText(Text);
-	UserWidget->Button->OnClicked.AddUniqueDynamic(this, &ThisClass::HandleClick);
+	UUserWidget* UserWidget = Cast<UUserWidget>(LocalGetUMGWidget());
+	if (ensure(UserWidget))
+	{
+		if (UMVVMView* View = UMVVMSubsystem::GetViewFromUserWidget(UserWidget))
+		{
+			View->SetViewModel(TEXT("Widget"), this);
+		}
+	}
 }
 
 
@@ -45,19 +50,17 @@ void UUIFrameworkSimpleButton::SetText(FText InText)
 	Text = InText;
 	MARK_PROPERTY_DIRTY_FROM_NAME(ThisClass, Text, this);
 	ForceNetUpdate();
+	UE_MVVM_BROADCAST_FIELD_VALUE_CHANGED(Text);
 }
 
 
 void UUIFrameworkSimpleButton::OnRep_Text()
 {
-	if (LocalGetUMGWidget())
-	{
-		CastChecked<UUIFrameworkSimpleButtonUserWidget>(LocalGetUMGWidget())->TextBlock->SetText(Text);
-	}
+	UE_MVVM_BROADCAST_FIELD_VALUE_CHANGED(Text);
 }
 
 
-void UUIFrameworkSimpleButton::HandleClick()
+void UUIFrameworkSimpleButton::OnClick(FMVVMEventField Field)
 {
 	// todo the click event should send the userid
 	ServerClick(Cast<APlayerController>(GetOuter()));
@@ -66,8 +69,7 @@ void UUIFrameworkSimpleButton::HandleClick()
 
 void UUIFrameworkSimpleButton::ServerClick_Implementation(APlayerController* PlayerController)
 {
-	FUIFrameworkClickEventArgument Argument;
-	Argument.PlayerController = PlayerController;
-	Argument.Sender = this;
-	OnClick.Broadcast(Argument);
+	ClickEvent.PlayerController = PlayerController;
+	ClickEvent.Sender = this;
+	BroadcastFieldValueChanged(ThisClass::FFieldNotificationClassDescriptor::ClickEvent);
 }
diff --git a/Engine/Plugins/Experimental/UIFramework/Source/Public/Components/UIFSimpleButtonUserWidget.h b/Engine/Plugins/Experimental/UIFramework/Source/Public/Components/UIFSimpleButtonUserWidget.h
deleted file mode 100644
diff --git a/Engine/Plugins/Experimental/UIFramework/Source/Public/UIFWidget.h b/Engine/Plugins/Experimental/UIFramework/Source/Public/UIFWidget.h
--- a/Engine/Plugins/Experimental/UIFramework/Source/Public/UIFWidget.h
+++ b/Engine/Plugins/Experimental/UIFramework/Source/Public/UIFWidget.h
@@ -4,6 +4,7 @@
 
 #include "Components/SlateWrapperTypes.h"
 #include "Components/Widget.h"
+#include "MVVMViewModelBase.h"
 #include "Templates/NonNullPointer.h"
 #include "Types/UIFParentWidget.h"
 #include "Types/UIFWidgetId.h"
@@ -37,7 +38,7 @@ public:
  * 
  */
 UCLASS(Abstract, BlueprintType)
-class UIFRAMEWORK_API UUIFrameworkWidget : public UObject
+class UIFRAMEWORK_API UUIFrameworkWidget : public UMVVMViewModelBase
 {
 	GENERATED_BODY()
 
diff --git a/Engine/Plugins/Experimental/UIFramework/Source/Public/Widgets/UIFSimpleButton.h b/Engine/Plugins/Experimental/UIFramework/Source/Public/Widgets/UIFSimpleButton.h
--- a/Engine/Plugins/Experimental/UIFramework/Source/Public/Widgets/UIFSimpleButton.h
+++ b/Engine/Plugins/Experimental/UIFramework/Source/Public/Widgets/UIFSimpleButton.h
@@ -3,6 +3,7 @@
 #pragma once
 
 #include "Types/UIFEvents.h"
+#include "Types/MVVMEventField.h"
 #include "UIFWidget.h"
 
 #include "UIFSimpleButton.generated.h"
@@ -18,31 +19,36 @@ class UIFRAMEWORK_API UUIFrameworkSimpleButton : public UUIFrameworkWidget
 public:
 	UUIFrameworkSimpleButton();
 
+private:
+	UPROPERTY(ReplicatedUsing=OnRep_Text, EditAnywhere, BlueprintReadOnly, Getter, Setter, FieldNotify, Category="UI Framework", meta = (AllowPrivateAccess = "true"))
+	FText Text;
+
+	UPROPERTY(BlueprintReadOnly, Getter, FieldNotify, Category = "UI Framework", meta = (DisallowedViewAccess = "true", AllowPrivateAccess = "true"))
+	FUIFrameworkClickEventArgument ClickEvent;
+
+public:
 	UFUNCTION(BlueprintCallable, BlueprintAuthorityOnly, Category = "UI Framework")
-	void SetText(FText Text);
+	void SetText(FText Value);
 
-	UFUNCTION(BlueprintCallable, Category = "UI Framework")
 	FText GetText() const
 	{
 		return Text;
 	}
 
+	const FUIFrameworkClickEventArgument& GetClickEvent() const
+	{
+		return ClickEvent;
+	}
+
 	virtual void LocalOnUMGWidgetCreated() override;
 
+	UFUNCTION(BlueprintCallable, Category = "UI Framework")
+	void OnClick(FMVVMEventField Field);
+
 private:
 	UFUNCTION()
 	void OnRep_Text();
 
-	UFUNCTION()
-	void HandleClick();
-
 	UFUNCTION(Server, Reliable)
 	void ServerClick(APlayerController* PlayerController);
-
-private:
-	UPROPERTY(ReplicatedUsing=OnRep_Text)
-	FText Text;
-
-public:
-	FUIFrameworkClickEvent OnClick;
 };
diff --git a/Engine/Plugins/Experimental/UIFramework/Source/UIFramework.Build.cs b/Engine/Plugins/Experimental/UIFramework/Source/UIFramework.Build.cs
--- a/Engine/Plugins/Experimental/UIFramework/Source/UIFramework.Build.cs
+++ b/Engine/Plugins/Experimental/UIFramework/Source/UIFramework.Build.cs
@@ -12,17 +12,18 @@ public class UIFramework : ModuleRules
 				"Core",
 				"CoreUObject",
 				"Engine",
+				"ModelViewViewModel",
 				"SlateCore",
 				"Slate",
 				"UMG",
 			}
 		);
 
-        PublicIncludePathModuleNames.AddRange(
-            new string[] {
-            }
-        );
-		
+		PublicIncludePathModuleNames.AddRange(
+			new string[] {
+			}
+		);
+
 		PrivateDependencyModuleNames.AddRange(
 			new string[]
 			{
diff --git a/Engine/Plugins/Experimental/UIFramework/UIFramework.uplugin b/Engine/Plugins/Experimental/UIFramework/UIFramework.uplugin
--- a/Engine/Plugins/Experimental/UIFramework/UIFramework.uplugin
+++ b/Engine/Plugins/Experimental/UIFramework/UIFramework.uplugin
@@ -20,5 +20,11 @@
 			"Type": "Runtime",
 			"LoadingPhase": "Default"
 		}
+	],
+	"Plugins": [
+		{
+			"Name": "ModelViewViewModel",
+			"Enabled": true
+		}
 	]
 }
\ No newline at end of file
