Add shadow, color and other properties to all UI text widget.

diff --git a/Engine/Plugins/Experimental/UIFramework/Source/Private/Widgets/UIFTextBlock.cpp b/Engine/Plugins/Experimental/UIFramework/Source/Private/Widgets/UIFTextBlock.cpp
--- a/Engine/Plugins/Experimental/UIFramework/Source/Private/Widgets/UIFTextBlock.cpp
+++ b/Engine/Plugins/Experimental/UIFramework/Source/Private/Widgets/UIFTextBlock.cpp
@@ -13,37 +13,42 @@
 /**
  *
  */
-UUIFrameworkTextBlock::UUIFrameworkTextBlock()
+UUIFrameworkTextBase::UUIFrameworkTextBase()
 {
-	WidgetClass = UTextBlock::StaticClass();
+	WidgetClass = nullptr;
 }
 
 
-void UUIFrameworkTextBlock::GetLifetimeReplicatedProps(TArray<FLifetimeProperty>& OutLifetimeProps) const
+void UUIFrameworkTextBase::GetLifetimeReplicatedProps(TArray<FLifetimeProperty>& OutLifetimeProps) const
 {
 	Super::GetLifetimeReplicatedProps(OutLifetimeProps);
 
 	FDoRepLifetimeParams Params;
 	Params.bIsPushBased = true;
 	DOREPLIFETIME_WITH_PARAMS_FAST(ThisClass, Text, Params);
+	DOREPLIFETIME_WITH_PARAMS_FAST(ThisClass, TextColor, Params);
 	DOREPLIFETIME_WITH_PARAMS_FAST(ThisClass, Justification, Params);
+	DOREPLIFETIME_WITH_PARAMS_FAST(ThisClass, OverflowPolicy, Params);
 }
 
-void UUIFrameworkTextBlock::LocalOnUMGWidgetCreated()
-{
-	UTextBlock* TextBlock = CastChecked<UTextBlock>(LocalGetUMGWidget());
-	TextBlock->SetText(Text);
-	TextBlock->SetJustification(Justification);
-}
-
-void UUIFrameworkTextBlock::SetText(FText InText)
+void UUIFrameworkTextBase::SetText(FText InText)
 {
 	Text = InText;
 	MARK_PROPERTY_DIRTY_FROM_NAME(ThisClass, Text, this);
 	ForceNetUpdate();
 }
 
-void UUIFrameworkTextBlock::SetJustification(ETextJustify::Type InJustification)
+void UUIFrameworkTextBase::SetTextColor(FLinearColor InTextColor)
+{
+	if (TextColor != InTextColor)
+	{
+		TextColor = InTextColor;
+		MARK_PROPERTY_DIRTY_FROM_NAME(ThisClass, TextColor, this);
+		ForceNetUpdate();
+	}
+}
+
+void UUIFrameworkTextBase::SetJustification(ETextJustify::Type InJustification)
 {
 	if (Justification != InJustification)
 	{
@@ -53,18 +58,149 @@ void UUIFrameworkTextBlock::SetJustification(ETextJustify::Type InJustification)
 	}
 }
 
-void UUIFrameworkTextBlock::OnRep_Text()
+void UUIFrameworkTextBase::SetOverflowPolicy(ETextOverflowPolicy InOverflowPolicy)
+{
+	if (OverflowPolicy != InOverflowPolicy)
+	{
+		OverflowPolicy = InOverflowPolicy;
+		MARK_PROPERTY_DIRTY_FROM_NAME(ThisClass, OverflowPolicy, this);
+		ForceNetUpdate();
+	}
+}
+
+void UUIFrameworkTextBase::LocalOnUMGWidgetCreated()
+{
+	SetTextToWidget(Text);
+	SetTextColorToWidget(TextColor);
+	SetJustificationToWidget(Justification);
+	SetOverflowPolicyToWidget(OverflowPolicy);
+}
+
+void UUIFrameworkTextBase::OnRep_Text()
+{
+	if (LocalGetUMGWidget())
+	{
+		SetTextToWidget(Text);
+	}
+}
+
+void UUIFrameworkTextBase::OnRep_TextColor()
+{
+	if (LocalGetUMGWidget())
+	{
+		SetTextColorToWidget(TextColor);
+	}
+}
+
+void UUIFrameworkTextBase::OnRep_Justification()
+{
+	if (LocalGetUMGWidget())
+	{
+		SetJustificationToWidget(Justification);
+	}
+}
+
+void UUIFrameworkTextBase::OnRep_OverflowPolicy()
+{
+	if (LocalGetUMGWidget())
+	{
+		SetOverflowPolicyToWidget(OverflowPolicy);
+	}
+}
+
+/**
+ * 
+ */
+UUIFrameworkTextBlock::UUIFrameworkTextBlock()
+{
+	WidgetClass = UTextBlock::StaticClass();
+}
+
+void UUIFrameworkTextBlock::GetLifetimeReplicatedProps(TArray<FLifetimeProperty>& OutLifetimeProps) const
+{
+	Super::GetLifetimeReplicatedProps(OutLifetimeProps);
+
+	FDoRepLifetimeParams Params;
+	Params.bIsPushBased = true;
+	DOREPLIFETIME_WITH_PARAMS_FAST(ThisClass, ShadowOffset, Params);
+	DOREPLIFETIME_WITH_PARAMS_FAST(ThisClass, ShadowColor, Params);
+}
+
+void UUIFrameworkTextBlock::SetShadowOffset(FVector2f InShadowOffset)
+{
+	if (ShadowOffset != InShadowOffset)
+	{
+		ShadowOffset = InShadowOffset;
+		MARK_PROPERTY_DIRTY_FROM_NAME(ThisClass, ShadowOffset, this);
+		ForceNetUpdate();
+	}
+}
+
+void UUIFrameworkTextBlock::SetShadowColor(FLinearColor InShadowColor)
+{
+	if (ShadowColor != InShadowColor)
+	{
+		ShadowColor = InShadowColor;
+		MARK_PROPERTY_DIRTY_FROM_NAME(ThisClass, ShadowColor, this);
+		ForceNetUpdate();
+	}
+}
+
+void UUIFrameworkTextBlock::LocalOnUMGWidgetCreated()
+{
+	Super::LocalOnUMGWidgetCreated();
+	SetShadowOffsetToWidget(ShadowOffset);
+	SetShadowColorToWidget(ShadowColor);
+}
+
+void UUIFrameworkTextBlock::OnRep_ShadowOffset()
+{
+	if (LocalGetUMGWidget())
+	{
+		SetShadowOffsetToWidget(ShadowOffset);
+	}
+}
+
+void UUIFrameworkTextBlock::OnRep_ShadowColor()
+{
+	if (LocalGetUMGWidget())
+	{
+		SetShadowColorToWidget(ShadowColor);
+	}
+}
+
+void UUIFrameworkTextBlock::SetTextToWidget(const FText& InValue)
+{
+	CastChecked<UTextBlock>(LocalGetUMGWidget())->SetText(InValue);
+}
+
+void UUIFrameworkTextBlock::SetTextColorToWidget(FLinearColor InValue)
+{
+	CastChecked<UTextBlock>(LocalGetUMGWidget())->SetColorAndOpacity(InValue);
+}
+
+void UUIFrameworkTextBlock::SetShadowOffsetToWidget(FVector2f InValue)
+{
+	CastChecked<UTextBlock>(LocalGetUMGWidget())->SetShadowOffset(FVector2d(InValue));
+}
+
+void UUIFrameworkTextBlock::SetShadowColorToWidget(FLinearColor InValue)
+{
+	CastChecked<UTextBlock>(LocalGetUMGWidget())->SetShadowColorAndOpacity(InValue);
+}
+
+void UUIFrameworkTextBlock::SetJustificationToWidget(ETextJustify::Type InJustification)
 {
 	if (LocalGetUMGWidget())
 	{
-		CastChecked<UTextBlock>(LocalGetUMGWidget())->SetText(Text);
+		CastChecked<UTextBlock>(LocalGetUMGWidget())->SetJustification(InJustification);
 	}
 }
 
-void UUIFrameworkTextBlock::OnRep_Justification()
+void UUIFrameworkTextBlock::SetOverflowPolicyToWidget(ETextOverflowPolicy InValue)
 {
 	if (LocalGetUMGWidget())
 	{
-		CastChecked<UTextBlock>(LocalGetUMGWidget())->SetJustification(Justification);
+		CastChecked<UTextBlock>(LocalGetUMGWidget())->SetTextOverflowPolicy(InValue);
 	}
 }
diff --git a/Engine/Plugins/Experimental/UIFramework/Source/Public/Widgets/UIFTextBlock.h b/Engine/Plugins/Experimental/UIFramework/Source/Public/Widgets/UIFTextBlock.h
--- a/Engine/Plugins/Experimental/UIFramework/Source/Public/Widgets/UIFTextBlock.h
+++ b/Engine/Plugins/Experimental/UIFramework/Source/Public/Widgets/UIFTextBlock.h
@@ -12,13 +12,13 @@ class UTextBlock;
 /**
  *
  */
-UCLASS(DisplayName = "TextBlock UIFramework")
-class UIFRAMEWORK_API UUIFrameworkTextBlock : public UUIFrameworkWidget
+UCLASS(Abstract, HideDropDown)
+class UIFRAMEWORK_API UUIFrameworkTextBase : public UUIFrameworkWidget
 {
 	GENERATED_BODY()
 
 public:
-	UUIFrameworkTextBlock();
+	UUIFrameworkTextBase();
 
 	UFUNCTION(BlueprintCallable, BlueprintAuthorityOnly, Category = "UI Framework")
 	void SetText(FText Text);
@@ -29,6 +29,15 @@ public:
 		return Text;
 	}
 	
+	UFUNCTION(BlueprintCallable, BlueprintAuthorityOnly, Category = "UI Framework")
+	void SetTextColor(FLinearColor TextColor);
+
+	UFUNCTION(BlueprintCallable, Category = "UI Framework")
+	FLinearColor GetTextColor() const
+	{
+		return TextColor;
+	}
+
 	UFUNCTION(BlueprintCallable, BlueprintAuthorityOnly, Category = "UI Framework")
 	void SetJustification(ETextJustify::Type Justification);
 
@@ -37,20 +46,108 @@ public:
 	{
 		return Justification;
 	}
+	
+	UFUNCTION(BlueprintCallable, BlueprintAuthorityOnly, Category = "UI Framework")
+	void SetOverflowPolicy(ETextOverflowPolicy OverflowPolicy);
+
+	UFUNCTION(BlueprintCallable, Category = "UI Framework")
+	ETextOverflowPolicy GetOverflowPolicy() const
+	{
+		return OverflowPolicy;
+	}
 
 	virtual void LocalOnUMGWidgetCreated() override;
 
 private:
 	UFUNCTION()
 	void OnRep_Text();
+	
+	UFUNCTION()
+	void OnRep_TextColor();
 
 	UFUNCTION()
 	void OnRep_Justification();
 
+	UFUNCTION()
+	void OnRep_OverflowPolicy();
+
+protected:
+	virtual void SetTextToWidget(const FText&)
+	{
+	}
+	virtual void SetTextColorToWidget(FLinearColor)
+	{
+	}
+	virtual void SetJustificationToWidget(ETextJustify::Type)
+	{
+	}
+	virtual void SetOverflowPolicyToWidget(ETextOverflowPolicy)
+	{
+	}
+
 private:
-	UPROPERTY(/*ExposeOnSpawn, */ReplicatedUsing=OnRep_Text)
+	UPROPERTY(/*ExposeOnSpawn, */ReplicatedUsing = OnRep_Text)
 	FText Text;
 
-	UPROPERTY(ReplicatedUsing=OnRep_Justification)
+	UPROPERTY(ReplicatedUsing = OnRep_TextColor)
+	FLinearColor TextColor = FLinearColor::Black;
+
+	UPROPERTY(ReplicatedUsing = OnRep_Justification)
 	TEnumAsByte<ETextJustify::Type> Justification;
+
+	UPROPERTY(ReplicatedUsing = OnRep_OverflowPolicy)
+	ETextOverflowPolicy OverflowPolicy = ETextOverflowPolicy::Clip;
+};
+
+/**
+ *
+ */
+UCLASS(DisplayName = "TextBlock UIFramework")
+class UIFRAMEWORK_API UUIFrameworkTextBlock : public UUIFrameworkTextBase
+{
+	GENERATED_BODY()
+
+public:
+	UUIFrameworkTextBlock();
+
+	UFUNCTION(BlueprintCallable, BlueprintAuthorityOnly, Category = "UI Framework")
+	void SetShadowOffset(FVector2f ShadowOffset);
+
+	UFUNCTION(BlueprintCallable, Category = "UI Framework")
+	FVector2f GetShadowOffset() const
+	{
+		return ShadowOffset;
+	}
+	
+	UFUNCTION(BlueprintCallable, BlueprintAuthorityOnly, Category = "UI Framework")
+	void SetShadowColor(FLinearColor ShadowColor);
+
+	UFUNCTION(BlueprintCallable, Category = "UI Framework")
+	FLinearColor GetShadowColor() const
+	{
+		return ShadowColor;
+	}
+
+	virtual void LocalOnUMGWidgetCreated() override;
+
+private:
+	UFUNCTION()
+	void OnRep_ShadowOffset();
+	UFUNCTION()
+	void OnRep_ShadowColor();
+
+protected:
+	virtual void SetTextToWidget(const FText&) override;
+	virtual void SetJustificationToWidget(ETextJustify::Type) override;
+	virtual void SetOverflowPolicyToWidget(ETextOverflowPolicy) override;
+	virtual void SetTextColorToWidget(FLinearColor) override;
+	virtual void SetShadowOffsetToWidget(FVector2f);
+	virtual void SetShadowColorToWidget(FLinearColor);
+
+private:
+	UPROPERTY(ReplicatedUsing = OnRep_ShadowOffset)
+	FVector2f ShadowOffset = FVector2f::ZeroVector;
+	
+	UPROPERTY(ReplicatedUsing = OnRep_ShadowColor)
+	FLinearColor ShadowColor = FLinearColor::Black;
 };
