Fix for removal of widget comparsion incorrectly checking parent child relation.

diff --git a/Engine/Plugins/Experimental/UIFramework/Source/Private/Widgets/UIFCanvasBox.cpp b/Engine/Plugins/Experimental/UIFramework/Source/Private/Widgets/UIFCanvasBox.cpp
--- a/Engine/Plugins/Experimental/UIFramework/Source/Private/Widgets/UIFCanvasBox.cpp
+++ b/Engine/Plugins/Experimental/UIFramework/Source/Private/Widgets/UIFCanvasBox.cpp
@@ -109,7 +109,7 @@ void UUIFrameworkCanvasBox::RemoveWidget(UUIFrameworkWidget* Widget)
 	{
 		FFrame::KismetExecutionMessage(TEXT("The widget is invalid. It can't be removed."), ELogVerbosity::Warning, "InvalidWidgetToRemove");
 	}
-	else if (Widget->AuthorityGetParent().IsParentValid() && Widget->AuthorityGetParent().IsWidget() && Widget->AuthorityGetParent().AsWidget() == this)
+	else if (Widget->AuthorityGetParent().IsParentValid() && Widget->AuthorityGetParent().IsWidget() && Widget->AuthorityGetParent().AsWidget() != this)
 	{
 		FFrame::KismetExecutionMessage(TEXT("The widget was created for another widget. It can't be removed on this player."), ELogVerbosity::Warning, "InvalidPlayerParentOnRemovedWidget");
 	}
diff --git a/Engine/Plugins/Experimental/UIFramework/Source/Private/Widgets/UIFStackBox.cpp b/Engine/Plugins/Experimental/UIFramework/Source/Private/Widgets/UIFStackBox.cpp
--- a/Engine/Plugins/Experimental/UIFramework/Source/Private/Widgets/UIFStackBox.cpp
+++ b/Engine/Plugins/Experimental/UIFramework/Source/Private/Widgets/UIFStackBox.cpp
@@ -115,7 +115,7 @@ void UUIFrameworkStackBox::RemoveWidget(UUIFrameworkWidget* Widget)
 	{
 		FFrame::KismetExecutionMessage(TEXT("The widget is invalid. It can't be removed."), ELogVerbosity::Warning, "InvalidWidgetToRemove");
 	}
-	else if (Widget->AuthorityGetParent().IsParentValid() && Widget->AuthorityGetParent().IsWidget() && Widget->AuthorityGetParent().AsWidget() == this)
+	else if (Widget->AuthorityGetParent().IsParentValid() && Widget->AuthorityGetParent().IsWidget() && Widget->AuthorityGetParent().AsWidget() != this)
 	{
 		FFrame::KismetExecutionMessage(TEXT("The widget was created for another widget. It can't be removed on this player."), ELogVerbosity::Warning, "InvalidPlayerParentOnRemovedWidget");
 	}
