More detailed error messages for incorrect widget child removal.

diff --git a/Engine/Plugins/Experimental/UIFramework/Source/Private/Widgets/UIFCanvasBox.cpp b/Engine/Plugins/Experimental/UIFramework/Source/Private/Widgets/UIFCanvasBox.cpp
--- a/Engine/Plugins/Experimental/UIFramework/Source/Private/Widgets/UIFCanvasBox.cpp
+++ b/Engine/Plugins/Experimental/UIFramework/Source/Private/Widgets/UIFCanvasBox.cpp
@@ -107,11 +107,19 @@ void UUIFrameworkCanvasBox::RemoveWidget(UUIFrameworkWidget* Widget)
 {
 	if (Widget == nullptr)
 	{
-		FFrame::KismetExecutionMessage(TEXT("The widget is invalid. It can't be removed."), ELogVerbosity::Warning, "InvalidWidgetToRemove");
+		FFrame::KismetExecutionMessage(TEXT("The widget is invalid. Widget can't be removed."), ELogVerbosity::Warning, "InvalidWidgetToRemove");
 	}
-	else if (Widget->AuthorityGetParent().IsParentValid() && Widget->AuthorityGetParent().IsWidget() && Widget->AuthorityGetParent().AsWidget() != this)
+	else if (!Widget->AuthorityGetParent().IsParentValid())
 	{
-		FFrame::KismetExecutionMessage(TEXT("The widget was created for another widget. It can't be removed on this player."), ELogVerbosity::Warning, "InvalidPlayerParentOnRemovedWidget");
+		FFrame::KismetExecutionMessage(TEXT("The widget parent is invalid. Widget can't be removed."), ELogVerbosity::Warning, "InvalidWidgetParentToRemoveFrom");
+	}
+	else if (!Widget->AuthorityGetParent().IsWidget())
+	{
+		FFrame::KismetExecutionMessage(TEXT("The widget parent is not a widget. Widget can't be removed."), ELogVerbosity::Warning, "NotWidgetParentToRemoveFrom");
+	}
+	else if (Widget->AuthorityGetParent().AsWidget() != this)
+	{
+		FFrame::KismetExecutionMessage(TEXT("The widget was added to another widget parent. It can't be removed on this player."), ELogVerbosity::Warning, "InvalidPlayerParentOnRemovedWidget");
 	}
 	else
 	{
diff --git a/Engine/Plugins/Experimental/UIFramework/Source/Private/Widgets/UIFStackBox.cpp b/Engine/Plugins/Experimental/UIFramework/Source/Private/Widgets/UIFStackBox.cpp
--- a/Engine/Plugins/Experimental/UIFramework/Source/Private/Widgets/UIFStackBox.cpp
+++ b/Engine/Plugins/Experimental/UIFramework/Source/Private/Widgets/UIFStackBox.cpp
@@ -113,11 +113,19 @@ void UUIFrameworkStackBox::RemoveWidget(UUIFrameworkWidget* Widget)
 {
 	if (Widget == nullptr)
 	{
-		FFrame::KismetExecutionMessage(TEXT("The widget is invalid. It can't be removed."), ELogVerbosity::Warning, "InvalidWidgetToRemove");
+		FFrame::KismetExecutionMessage(TEXT("The widget is invalid. Widget can't be removed."), ELogVerbosity::Warning, "InvalidWidgetToRemove");
 	}
-	else if (Widget->AuthorityGetParent().IsParentValid() && Widget->AuthorityGetParent().IsWidget() && Widget->AuthorityGetParent().AsWidget() != this)
+	else if (!Widget->AuthorityGetParent().IsParentValid())
 	{
-		FFrame::KismetExecutionMessage(TEXT("The widget was created for another widget. It can't be removed on this player."), ELogVerbosity::Warning, "InvalidPlayerParentOnRemovedWidget");
+		FFrame::KismetExecutionMessage(TEXT("The widget parent is invalid. Widget can't be removed."), ELogVerbosity::Warning, "InvalidWidgetParentToRemoveFrom");
+	}
+	else if (!Widget->AuthorityGetParent().IsWidget())
+	{
+		FFrame::KismetExecutionMessage(TEXT("The widget parent is not a widget. Widget can't be removed."), ELogVerbosity::Warning, "NotWidgetParentToRemoveFrom");
+	}
+	else if (Widget->AuthorityGetParent().AsWidget() != this)
+	{
+		FFrame::KismetExecutionMessage(TEXT("The widget was added to another widget parent. It can't be removed on this player."), ELogVerbosity::Warning, "InvalidPlayerParentOnRemovedWidget");
 	}
 	else
 	{
