VerseUIOverlay widget.

#tests wrote a small script that sets up a canvas with a couple of buttons for adding and removing widgets to an overlay. Tested with PIE.

diff --git a/Engine/Plugins/Experimental/UIFramework/Source/Private/Widgets/UIFOverlay.cpp b/Engine/Plugins/Experimental/UIFramework/Source/Private/Widgets/UIFOverlay.cpp
new file mode 100644
--- /dev/null
+++ b/Engine/Plugins/Experimental/UIFramework/Source/Private/Widgets/UIFOverlay.cpp
@@ -0,0 +1,227 @@
+// Copyright Epic Games, Inc. All Rights Reserved.
+
+#include "Widgets/UIFOverlay.h"
+#include "Types/UIFWidgetTree.h"
+#include "UIFLog.h"
+#include "UIFModule.h"
+#include "UIFPlayerComponent.h"
+
+#include "Components/Spacer.h"
+#include "Components/Overlay.h"
+#include "Components/OverlaySlot.h"
+
+#include "Engine/ActorChannel.h"
+#include "Engine/Engine.h"
+#include "Engine/NetDriver.h"
+#include "GameFramework/Actor.h"
+#include "GameFramework/PlayerController.h"
+#include "Net/UnrealNetwork.h"
+#include "UObject/Package.h"
+
+#include UE_INLINE_GENERATED_CPP_BY_NAME(UIFOverlay)
+
+/**
+ *
+ */
+void FUIFrameworkOverlaySlotList::PostReplicatedChange(const TArrayView<int32>& ChangedIndices, int32 FinalSize)
+{
+	check(Owner);
+	for (int32 Index : ChangedIndices)
+	{
+		FUIFrameworkOverlaySlot& Slot = Slots[Index];
+		if (Slot.LocalIsAquiredWidgetValid())
+		{
+			Owner->LocalAddChild(Slot.GetWidgetId());
+		}
+	}
+}
+
+bool FUIFrameworkOverlaySlotList::NetDeltaSerialize(FNetDeltaSerializeInfo& DeltaParms)
+{
+	return FFastArraySerializer::FastArrayDeltaSerialize<FUIFrameworkOverlaySlot, FUIFrameworkOverlaySlotList>(Slots, DeltaParms, *this);
+}
+
+void FUIFrameworkOverlaySlotList::AddEntry(FUIFrameworkOverlaySlot Entry)
+{
+	int32 NewEntryIndex = Slots.Add(MoveTemp(Entry));
+	FUIFrameworkOverlaySlot& NewEntry = Slots[NewEntryIndex];
+	NewEntry.Index = NewEntryIndex;
+	MarkItemDirty(NewEntry);
+}
+
+bool FUIFrameworkOverlaySlotList::RemoveEntry(UUIFrameworkWidget* Widget)
+{
+	check(Widget);
+	int32 Index = Slots.IndexOfByPredicate([Widget](const FUIFrameworkOverlaySlot& Entry) { return Entry.AuthorityGetWidget() == Widget; });
+	if (Index != INDEX_NONE)
+	{
+		Slots.RemoveAt(Index);
+		for (; Index < Slots.Num(); ++Index)
+		{
+			Slots[Index].Index = Index;
+		}
+		MarkArrayDirty();
+	}
+	return Index != INDEX_NONE;
+}
+
+FUIFrameworkOverlaySlot* FUIFrameworkOverlaySlotList::FindEntry(FUIFrameworkWidgetId WidgetId)
+{
+	return Slots.FindByPredicate([WidgetId](const FUIFrameworkOverlaySlot& Entry) { return Entry.GetWidgetId() == WidgetId; });
+}
+
+void FUIFrameworkOverlaySlotList::ForEachChildren(const TFunctionRef<void(UUIFrameworkWidget*)>& Func)
+{
+	for (FUIFrameworkOverlaySlot& Slot : Slots)
+	{
+		Func(Slot.AuthorityGetWidget());
+	}
+}
+
+/**
+ *
+ */
+UUIFrameworkOverlay::UUIFrameworkOverlay()
+	: ReplicatedSlotList(this)
+{
+	WidgetClass = UOverlay::StaticClass();
+}
+
+void UUIFrameworkOverlay::GetLifetimeReplicatedProps(TArray<FLifetimeProperty>& OutLifetimeProps) const
+{
+	Super::GetLifetimeReplicatedProps(OutLifetimeProps);
+
+	FDoRepLifetimeParams Params;
+	Params.bIsPushBased = true;
+	DOREPLIFETIME_WITH_PARAMS_FAST(ThisClass, ReplicatedSlotList, Params);
+}
+
+void UUIFrameworkOverlay::AddWidget(FUIFrameworkOverlaySlot InEntry)
+{
+	if (InEntry.AuthorityGetWidget() == nullptr)
+	{
+		FFrame::KismetExecutionMessage(TEXT("The widget is invalid. It can't be added."), ELogVerbosity::Warning, "InvalidWidgetToAdd");
+	}
+	else
+	{
+		// Reset the widget to make sure the id is set and it may have been duplicated during the attach
+		InEntry.AuthoritySetWidget(FUIFrameworkModule::AuthorityAttachWidget(this, InEntry.AuthorityGetWidget()));
+		AddEntry(MoveTemp(InEntry));
+	}
+}
+
+void UUIFrameworkOverlay::RemoveWidget(UUIFrameworkWidget* Widget)
+{
+	if (Widget == nullptr)
+	{
+		FFrame::KismetExecutionMessage(TEXT("The widget is invalid. Widget can't be removed."), ELogVerbosity::Warning, "InvalidWidgetToRemove");
+	}
+	else if (!Widget->AuthorityGetParent().IsParentValid())
+	{
+		FFrame::KismetExecutionMessage(TEXT("The widget parent is invalid. Widget can't be removed."), ELogVerbosity::Warning, "InvalidWidgetParentToRemoveFrom");
+	}
+	else if (!Widget->AuthorityGetParent().IsWidget())
+	{
+		FFrame::KismetExecutionMessage(TEXT("The widget parent is not a widget. Widget can't be removed."), ELogVerbosity::Warning, "NotWidgetParentToRemoveFrom");
+	}
+	else if (Widget->AuthorityGetParent().AsWidget() != this)
+	{
+		FFrame::KismetExecutionMessage(TEXT("The widget was added to another widget parent. It can't be removed on this player."), ELogVerbosity::Warning, "InvalidPlayerParentOnRemovedWidget");
+	}
+	else
+	{
+		FUIFrameworkModule::AuthorityDetachWidgetFromParent(Widget);
+		RemoveEntry(Widget);
+	}
+}
+
+void UUIFrameworkOverlay::AuthorityForEachChildren(const TFunctionRef<void(UUIFrameworkWidget*)>& Func)
+{
+	Super::AuthorityForEachChildren(Func);
+
+	ReplicatedSlotList.ForEachChildren(Func);
+}
+
+void UUIFrameworkOverlay::AuthorityRemoveChild(UUIFrameworkWidget* Widget)
+{
+	Super::AuthorityRemoveChild(Widget);
+	RemoveEntry(Widget);
+}
+
+void UUIFrameworkOverlay::LocalAddChild(FUIFrameworkWidgetId ChildId)
+{
+	bool bIsAdded = false;
+	if (FUIFrameworkOverlaySlot* OverlayEntry = FindEntry(ChildId))
+	{
+		if (FUIFrameworkWidgetTree* WidgetTree = GetWidgetTree())
+		{
+			if (UUIFrameworkWidget* ChildWidget = WidgetTree->FindWidgetById(ChildId))
+			{
+				UWidget* ChildUMGWidget = ChildWidget->LocalGetUMGWidget();
+				if (ensure(ChildUMGWidget))
+				{
+					UOverlay* Overlay = CastChecked<UOverlay>(LocalGetUMGWidget());
+					UPanelSlot* PanelSlot = nullptr;
+
+					const int32 ChildCount = Overlay->GetChildrenCount();
+					if (OverlayEntry->Index < ChildCount)
+					{
+						Overlay->ReplaceOverlayChildAt(OverlayEntry->Index, ChildUMGWidget);
+						PanelSlot = Overlay->GetSlots()[OverlayEntry->Index];
+					}
+					else if (OverlayEntry->Index == ChildCount)
+					{
+						PanelSlot = Overlay->AddChild(ChildUMGWidget);
+					}
+					else
+					{
+						for (int32 Index = ChildCount; Index < OverlayEntry->Index; ++Index)
+						{
+							Overlay->AddChild(NewObject<USpacer>(GetTransientPackage()));
+						}
+						PanelSlot = Overlay->AddChild(ChildUMGWidget);
+					}
+
+					checkf(PanelSlot, TEXT("OverlayPanel should be able to receive slot"));
+					
+					OverlayEntry->LocalAquireWidget();
+
+					UOverlaySlot* OverlaySlot = CastChecked<UOverlaySlot>(PanelSlot);
+					OverlaySlot->SetPadding(OverlayEntry->Padding);
+					OverlaySlot->SetVerticalAlignment(OverlayEntry->VerticalAlignment);
+					OverlaySlot->SetHorizontalAlignment(OverlayEntry->HorizontalAlignment);
+				}
+			}
+			else
+			{
+				UE_LOG(LogUIFramework, Log, TEXT("The widget '%" INT64_FMT "' doesn't exist in the WidgetTree."), ChildId.GetKey());
+				Super::LocalAddChild(ChildId);
+			}
+		}
+		else
+		{
+			UE_LOG(LogUIFramework, Log, TEXT("The widget '%" INT64_FMT "' doesn't exist in the WidgetTree."), ChildId.GetKey());
+			Super::LocalAddChild(ChildId);
+		}
+	}
+	else
+	{
+		UE_LOG(LogUIFramework, Verbose, TEXT("The widget '%" INT64_FMT "' was not found in the Overlay Slots."), ChildId.GetKey());
+		Super::LocalAddChild(ChildId);
+	}
+}
+
+void UUIFrameworkOverlay::AddEntry(FUIFrameworkOverlaySlot Entry)
+{
+	ReplicatedSlotList.AddEntry(Entry);
+}
+
+bool UUIFrameworkOverlay::RemoveEntry(UUIFrameworkWidget* Widget)
+{
+	return ReplicatedSlotList.RemoveEntry(Widget);
+}
+
+FUIFrameworkOverlaySlot* UUIFrameworkOverlay::FindEntry(FUIFrameworkWidgetId WidgetId)
+{
+	return ReplicatedSlotList.FindEntry(WidgetId);
+}
diff --git a/Engine/Plugins/Experimental/UIFramework/Source/Public/Widgets/UIFOverlay.h b/Engine/Plugins/Experimental/UIFramework/Source/Public/Widgets/UIFOverlay.h
new file mode 100644
--- /dev/null
+++ b/Engine/Plugins/Experimental/UIFramework/Source/Public/Widgets/UIFOverlay.h
@@ -0,0 +1,118 @@
+// Copyright Epic Games, Inc. All Rights Reserved.
+
+#pragma once
+
+#include "Components/SlateWrapperTypes.h"
+#include "Layout/Margin.h"
+#include "Types/SlateEnums.h"
+#include "Types/UIFSlotBase.h"
+#include "UIFWidget.h"
+#include "Widgets/Layout/Anchors.h"
+
+#include "UIFOverlay.generated.h"
+
+class UUIFrameworkOverlay;
+struct FUIFrameworkOverlaySlotList;
+
+/**
+ *
+ */
+USTRUCT(BlueprintType)
+struct FUIFrameworkOverlaySlot : public FUIFrameworkSlotBase
+{
+	GENERATED_BODY()
+
+	friend UUIFrameworkOverlay;
+	friend FUIFrameworkOverlaySlotList;
+
+	/** Distance between that surrounds the widget inside the slot. */
+	UPROPERTY(EditAnywhere, BlueprintReadWrite, Category = "UI Framework")
+	FMargin Padding = FMargin(0.f, 0.f, 0.f, 0.f);
+
+	/** Horizontal alignment of the widget inside the slot. */
+	UPROPERTY(EditAnywhere, BlueprintReadWrite, Category = "UI Framework")
+	TEnumAsByte<EHorizontalAlignment> HorizontalAlignment = EHorizontalAlignment::HAlign_Fill;
+
+	/** Vertical alignment of the widget inside the slot. */
+	UPROPERTY(EditAnywhere, BlueprintReadWrite, Category = "UI Framework")
+	TEnumAsByte<EVerticalAlignment> VerticalAlignment = EVerticalAlignment::VAlign_Fill;
+
+private:
+	/** Index in the array the Slot is. The position in the array can change when replicated. */
+	UPROPERTY()
+	int32 Index;
+};
+
+
+/**
+ *
+ */
+USTRUCT()
+struct UIFRAMEWORK_API FUIFrameworkOverlaySlotList : public FFastArraySerializer
+{
+	GENERATED_BODY()
+
+	FUIFrameworkOverlaySlotList() = default;
+	FUIFrameworkOverlaySlotList(UUIFrameworkOverlay* InOwner)
+		: Owner(InOwner)
+	{}
+
+public:
+	//~ Begin of FFastArraySerializer
+	void PostReplicatedChange(const TArrayView<int32>& ChangedIndices, int32 FinalSize);
+	//~ End of FFastArraySerializer
+
+	bool NetDeltaSerialize(FNetDeltaSerializeInfo& DeltaParms);
+
+	void AddEntry(FUIFrameworkOverlaySlot Entry);
+	bool RemoveEntry(UUIFrameworkWidget* Widget);
+	FUIFrameworkOverlaySlot* FindEntry(FUIFrameworkWidgetId WidgetId);
+	void ForEachChildren(const TFunctionRef<void(UUIFrameworkWidget*)>& Func);
+
+private:
+	UPROPERTY()
+	TArray<FUIFrameworkOverlaySlot> Slots;
+
+	UPROPERTY(NotReplicated, Transient)
+	TObjectPtr<UUIFrameworkOverlay> Owner;
+};
+
+
+template<>
+struct TStructOpsTypeTraits<FUIFrameworkOverlaySlotList> : public TStructOpsTypeTraitsBase2<FUIFrameworkOverlaySlotList>
+{
+	enum { WithNetDeltaSerializer = true };
+};
+
+
+/**
+ *
+ */
+UCLASS(DisplayName="Overlay UIFramework")
+class UIFRAMEWORK_API UUIFrameworkOverlay : public UUIFrameworkWidget
+{
+	GENERATED_BODY()
+
+	friend FUIFrameworkOverlaySlotList;
+
+public:
+	UUIFrameworkOverlay();
+
+	UFUNCTION(BlueprintCallable, BlueprintAuthorityOnly, Category = "UI Framework")
+	void AddWidget(FUIFrameworkOverlaySlot Widget);
+
+	UFUNCTION(BlueprintCallable, BlueprintAuthorityOnly, Category = "UI Framework")
+	void RemoveWidget(UUIFrameworkWidget* Widget);
+
+	virtual void AuthorityForEachChildren(const TFunctionRef<void(UUIFrameworkWidget*)>& Func) override;
+	virtual void AuthorityRemoveChild(UUIFrameworkWidget* Widget) override;
+	virtual void LocalAddChild(FUIFrameworkWidgetId ChildId) override;
+
+private:
+	void AddEntry(FUIFrameworkOverlaySlot Entry);
+	bool RemoveEntry(UUIFrameworkWidget* Widget);
+	FUIFrameworkOverlaySlot* FindEntry(FUIFrameworkWidgetId WidgetId);
+
+	UPROPERTY(Replicated)
+	FUIFrameworkOverlaySlotList ReplicatedSlotList;
+};
diff --git a/Engine/Source/Runtime/UMG/Private/Components/Overlay.cpp b/Engine/Source/Runtime/UMG/Private/Components/Overlay.cpp
--- a/Engine/Source/Runtime/UMG/Private/Components/Overlay.cpp
+++ b/Engine/Source/Runtime/UMG/Private/Components/Overlay.cpp
@@ -56,6 +56,18 @@ void UOverlay::OnSlotRemoved(UPanelSlot* InSlot)
 	}
 }
 
+bool UOverlay::ReplaceOverlayChildAt(int32 Index, UWidget* Content)
+{
+	if (!Slots.IsValidIndex(Index) || Content == nullptr)
+	{
+		return false;
+	}
+
+	UOverlaySlot* OverlaySlot = CastChecked<UOverlaySlot>(Slots[Index]);
+	OverlaySlot->ReplaceContent(Content);
+	return true;
+}
+
 TSharedRef<SWidget> UOverlay::RebuildWidget()
 {
 	MyOverlay = SNew(SOverlay);
diff --git a/Engine/Source/Runtime/UMG/Private/Components/OverlaySlot.cpp b/Engine/Source/Runtime/UMG/Private/Components/OverlaySlot.cpp
--- a/Engine/Source/Runtime/UMG/Private/Components/OverlaySlot.cpp
+++ b/Engine/Source/Runtime/UMG/Private/Components/OverlaySlot.cpp
@@ -36,6 +36,36 @@ void UOverlaySlot::BuildSlot(TSharedRef<SOverlay> Overlay)
 		];
 }
 
+void UOverlaySlot::ReplaceContent(UWidget* NewContent)
+{
+	if (Content != NewContent)
+	{
+		if (NewContent)
+		{
+			NewContent->RemoveFromParent();
+		}
+
+		if (UWidget* PreviousWidget = Content)
+		{
+			// Setting Slot=null before RemoveFromParent to prevent destroying this slot
+			PreviousWidget->Slot = nullptr;
+			PreviousWidget->RemoveFromParent();
+		}
+
+		Content = NewContent;
+
+		if (Content)
+		{
+			Content->Slot = this;
+		}
+
+		if (Slot)
+		{
+			Slot->AttachWidget(Content == nullptr ? SNullWidget::NullWidget : Content->TakeWidget());
+		}
+	}
+}
+
 void UOverlaySlot::SetPadding(FMargin InPadding)
 {
 	Padding = InPadding;
diff --git a/Engine/Source/Runtime/UMG/Public/Components/Overlay.h b/Engine/Source/Runtime/UMG/Public/Components/Overlay.h
--- a/Engine/Source/Runtime/UMG/Public/Components/Overlay.h
+++ b/Engine/Source/Runtime/UMG/Public/Components/Overlay.h
@@ -25,6 +25,10 @@ class UMG_API UOverlay : public UPanelWidget
 	UFUNCTION(BlueprintCallable, Category="Widget")
 	UOverlaySlot* AddChildToOverlay(UWidget* Content);
 
+	/** Replace the widget at the given index it with a different widget. */
+	UFUNCTION(BlueprintCallable, Category = "Panel")
+	bool ReplaceOverlayChildAt(int32 Index, UWidget* Content);
+
 #if WITH_EDITOR
 	virtual const FText GetPaletteCategory() override;
 #endif
diff --git a/Engine/Source/Runtime/UMG/Public/Components/OverlaySlot.h b/Engine/Source/Runtime/UMG/Public/Components/OverlaySlot.h
--- a/Engine/Source/Runtime/UMG/Public/Components/OverlaySlot.h
+++ b/Engine/Source/Runtime/UMG/Public/Components/OverlaySlot.h
@@ -73,5 +73,8 @@ public:
 	/** Builds the underlying FSlot for the Slate layout panel. */
 	virtual void BuildSlot(TSharedRef<SOverlay> InOverlay);
 
+	/** Replace the slot content. */
+	void ReplaceContent(UWidget* Content);
+
 	virtual void ReleaseSlateResources(bool bReleaseChildren) override;
 };
