VerseUI: Add texture and material widget.

diff --git a/Engine/Plugins/Experimental/UIFramework/Source/Private/UIFLocalSettings.cpp b/Engine/Plugins/Experimental/UIFramework/Source/Private/UIFLocalSettings.cpp
new file mode 100644
--- /dev/null
+++ b/Engine/Plugins/Experimental/UIFramework/Source/Private/UIFLocalSettings.cpp
@@ -0,0 +1,51 @@
+// Copyright Epic Games, Inc. All Rights Reserved.
+
+#include "UIFLocalSettings.h"
+
+#include "Async/Async.h"
+#include "Engine/AssetManager.h"
+#include "Engine/StreamableManager.h"
+
+UUIFrameworkLocalSettings::UUIFrameworkLocalSettings()
+{
+	ErrorResource = FSoftObjectPath(TEXT("/UIFramework/Materials/T_Icon_Warning_32.T_Icon_Warning_32"));
+	LoadingResource = FSoftObjectPath(TEXT("/UIFramework/Materials/M_UI_DefaultLoading.M_UI_DefaultLoading"));
+}
+
+void UUIFrameworkLocalSettings::LoadResources()
+{
+	Async(EAsyncExecution::TaskGraphMainThread, []()
+		{
+			const UUIFrameworkLocalSettings* Settings = GetDefault<UUIFrameworkLocalSettings>();
+			TArray<FSoftObjectPath> ObjectsToLoad;
+			ObjectsToLoad.Reserve(2);
+			ObjectsToLoad.Add(Settings->ErrorResource.ToSoftObjectPath());
+			ObjectsToLoad.Add(Settings->LoadingResource.ToSoftObjectPath());
+			UAssetManager::GetStreamableManager().RequestAsyncLoad(
+				ObjectsToLoad,
+				[]()
+				{
+					if (UObjectInitialized() && !IsEngineExitRequested())
+					{
+						UUIFrameworkLocalSettings* Default = GetMutableDefault<UUIFrameworkLocalSettings>();
+						Default->ErrorResourcePtr = Default->ErrorResource.Get();
+						Default->LoadingResourcePtr = Default->LoadingResource.Get();
+					}
+				},
+				FStreamableManager::DefaultAsyncLoadPriority);
+		});
+}
+
+
+FName UUIFrameworkLocalSettings::GetCategoryName() const
+{
+	return TEXT("Plugins");
+}
+
+#if WITH_EDITOR
+void UUIFrameworkLocalSettings::PostEditChangeChainProperty(struct FPropertyChangedChainEvent& PropertyChangedEvent)
+{
+	Super::PostEditChangeChainProperty(PropertyChangedEvent);
+	LoadResources();
+}
+#endif
\ No newline at end of file
diff --git a/Engine/Plugins/Experimental/UIFramework/Source/Private/UIFModule.cpp b/Engine/Plugins/Experimental/UIFramework/Source/Private/UIFModule.cpp
--- a/Engine/Plugins/Experimental/UIFramework/Source/Private/UIFModule.cpp
+++ b/Engine/Plugins/Experimental/UIFramework/Source/Private/UIFModule.cpp
@@ -3,8 +3,10 @@
 #include "UIFModule.h"
 
 #include "Modules/ModuleManager.h"
+#include "UIFLocalSettings.h"
 #include "UIFPresenter.h"
 #include "UIFWidget.h"
+#include "UObject/Package.h"
 
 namespace UE::UIFramework::Private
 {
@@ -160,4 +162,9 @@ TSubclassOf<UUIFrameworkPresenter> FUIFrameworkModule::GetPresenterClass()
 	return UE::UIFramework::Private::Director.Get() ? UE::UIFramework::Private::Director : TSubclassOf<UUIFrameworkPresenter>(UUIFrameworkGameViewportPresenter::StaticClass());
 }
 
-IMPLEMENT_MODULE(FUIFrameworkModule, UIFramework)
+void FUIFrameworkModule::StartupModule()
+{
+	GetMutableDefault<UUIFrameworkLocalSettings>()->LoadResources();
+}
+
+IMPLEMENT_MODULE(FUIFrameworkModule, UIFramework)
\ No newline at end of file
diff --git a/Engine/Plugins/Experimental/UIFramework/Source/Private/Widgets/UIFImageBlock.cpp b/Engine/Plugins/Experimental/UIFramework/Source/Private/Widgets/UIFImageBlock.cpp
new file mode 100644
--- /dev/null
+++ b/Engine/Plugins/Experimental/UIFramework/Source/Private/Widgets/UIFImageBlock.cpp
@@ -0,0 +1,183 @@
+// Copyright Epic Games, Inc. All Rights Reserved.
+
+#include "Widgets/UIFImageBlock.h"
+
+#include "Components/Image.h"
+#include "Engine/AssetManager.h"
+#include "Engine/StreamableManager.h"
+#include "Engine/Texture2D.h"
+#include "Materials/Material.h"
+#include "Materials/MaterialInterface.h"
+#include "UIFLocalSettings.h"
+
+#include "Net/Core/PushModel/PushModel.h"
+#include "Net/UnrealNetwork.h"
+
+#include UE_INLINE_GENERATED_CPP_BY_NAME(UIFImageBlock)
+
+/**
+ *
+ */
+UUIFrameworkImageBlock::UUIFrameworkImageBlock()
+{
+	WidgetClass = UImage::StaticClass();
+}
+
+void UUIFrameworkImageBlock::GetLifetimeReplicatedProps(TArray<FLifetimeProperty>& OutLifetimeProps) const
+{
+	Super::GetLifetimeReplicatedProps(OutLifetimeProps);
+
+	FDoRepLifetimeParams Params;
+	Params.bIsPushBased = true;
+	DOREPLIFETIME_WITH_PARAMS_FAST(ThisClass, Data, Params);
+}
+
+void UUIFrameworkImageBlock::SetMaterial(TSoftObjectPtr<UMaterialInterface> InMaterial)
+{
+	if (Data.ResourceObject != InMaterial)
+	{
+		Data.ResourceObject = InMaterial;
+		Data.bUseTextureSize = false;
+		MARK_PROPERTY_DIRTY_FROM_NAME(ThisClass, Data, this);
+		ForceNetUpdate();
+	}
+}
+
+void UUIFrameworkImageBlock::SetTexture(TSoftObjectPtr<UTexture2D> InTexture, bool bUseTextureSize)
+{
+	if (Data.ResourceObject != InTexture)
+	{
+		Data.ResourceObject = InTexture;
+		Data.bUseTextureSize = bUseTextureSize;
+		MARK_PROPERTY_DIRTY_FROM_NAME(ThisClass, Data, this);
+		ForceNetUpdate();
+	}
+}
+
+void UUIFrameworkImageBlock::SetTint(FLinearColor InTint)
+{
+	if (Data.Tint != InTint)
+	{
+		Data.Tint = InTint;
+		MARK_PROPERTY_DIRTY_FROM_NAME(ThisClass, Data, this);
+		ForceNetUpdate();
+	}
+}
+
+void UUIFrameworkImageBlock::SetTiling(ESlateBrushTileType::Type InTiling)
+{
+	if (Data.Tiling != InTiling)
+	{
+		Data.Tiling = InTiling;
+		MARK_PROPERTY_DIRTY_FROM_NAME(ThisClass, Data, this);
+		ForceNetUpdate();
+	}
+}
+
+void UUIFrameworkImageBlock::SetDesiredSize(FVector2f InDesiredSize)
+{
+	if (Data.DesiredSize != InDesiredSize)
+	{
+		Data.DesiredSize = InDesiredSize;
+		MARK_PROPERTY_DIRTY_FROM_NAME(ThisClass, Data, this);
+		ForceNetUpdate();
+	}
+}
+
+void UUIFrameworkImageBlock::LocalOnUMGWidgetCreated()
+{
+	UImage* Image = CastChecked<UImage>(LocalGetUMGWidget());
+
+	FSlateBrush TmpBrush = Image->Brush;
+	TmpBrush.ImageSize = FVector2D(Data.DesiredSize);
+	TmpBrush.TintColor = Data.Tint;
+	TmpBrush.Tiling = Data.Tiling;
+
+	UObject* Resource = ProcessResourceObject();
+	if (Data.bUseTextureSize)
+	{
+		if (UTexture2D* Texture = Cast<UTexture2D>(Resource))
+		{
+			TmpBrush.ImageSize.X = Texture->GetSizeX();
+			TmpBrush.ImageSize.Y = Texture->GetSizeY();
+		}
+	}
+	TmpBrush.SetResourceObject(Resource);
+
+	Image->SetBrush(TmpBrush);
+}
+
+bool UUIFrameworkImageBlock::LocalIsReplicationReady() const
+{
+	return Super::LocalIsReplicationReady() && (!bWaitForResourceToBeLoaded || Data.ResourceObject.IsValid() || Data.ResourceObject.IsNull());
+}
+
+void UUIFrameworkImageBlock::OnRep_Data()
+{
+	if (LocalGetUMGWidget())
+	{
+		UUIFrameworkImageBlock::LocalOnUMGWidgetCreated();
+	}
+	else if (bWaitForResourceToBeLoaded)
+	{
+		ProcessResourceObject();
+	}
+}
+
+UObject* UUIFrameworkImageBlock::ProcessResourceObject()
+{
+	++StreamingCounter;
+	if (StreamingHandle.IsValid())
+	{
+		StreamingHandle->CancelHandle();
+		StreamingHandle.Reset();
+	}
+
+	if (Data.ResourceObject.IsNull())
+	{
+		return nullptr;
+	}
+	else if (Data.ResourceObject.IsPending())
+	{
+		TWeakObjectPtr<UUIFrameworkImageBlock> WeakThis(this);
+		int32 CurrentStreamingCounter = StreamingCounter;
+		StreamingHandle = UAssetManager::GetStreamableManager().RequestAsyncLoad(
+			Data.ResourceObject.ToSoftObjectPath(),
+			[WeakThis, CurrentStreamingCounter]() {
+				if (UUIFrameworkImageBlock* StrongThis = WeakThis.Get())
+				{
+					// If the StreamingCounter do not matches, then this delegate was interrupted, but had already been queued for a callback.
+					if (StrongThis->StreamingCounter != CurrentStreamingCounter)
+					{
+						return;
+					}
+
+					StrongThis->OnRep_Data();
+				}
+			},
+			FStreamableManager::AsyncLoadHighPriority);
+
+		return GetDefault<UUIFrameworkLocalSettings>()->GetLoadingResource();
+	}
+	else
+	{
+		check(Data.ResourceObject.IsValid());
+		UObject* ResourceObject = Data.ResourceObject.Get();
+		UObject* Result = GetDefault<UUIFrameworkLocalSettings>()->GetErrorResource();
+		if (UMaterialInterface* MaterialInterface = Cast<UMaterialInterface>(ResourceObject))
+		{
+			if (UMaterial* BaseMaterial = MaterialInterface->GetBaseMaterial())
+			{
+				if (BaseMaterial->IsUIMaterial())
+				{
+					Result = ResourceObject;
+				}
+			}
+		}
+		else if (UTexture2D* Texture = Cast<UTexture2D>(ResourceObject))
+		{
+			Result = ResourceObject;
+		}
+		return Result;
+	}
+}
diff --git a/Engine/Plugins/Experimental/UIFramework/Source/Public/UIFLocalSettings.h b/Engine/Plugins/Experimental/UIFramework/Source/Public/UIFLocalSettings.h
new file mode 100644
--- /dev/null
+++ b/Engine/Plugins/Experimental/UIFramework/Source/Public/UIFLocalSettings.h
@@ -0,0 +1,64 @@
+// Copyright Epic Games, Inc. All Rights Reserved.
+
+#pragma once
+
+#include "CoreMinimal.h"
+#include "Engine/DeveloperSettings.h"
+#include "UIFLocalSettings.generated.h"
+
+class UObject;
+
+/**
+ * 
+ */
+UCLASS(config=Game, defaultconfig, meta = (DisplayName = "UI Framework Local Settings"))
+class UIFRAMEWORK_API UUIFrameworkLocalSettings : public UDeveloperSettings
+{
+	GENERATED_BODY()
+
+public:
+	UUIFrameworkLocalSettings();
+
+	UObject* GetErrorResource() const
+	{
+		return ErrorResourcePtr;
+	}
+
+	UObject* GetLoadingResource() const
+	{
+		return LoadingResourcePtr;
+	}
+
+	void LoadResources();
+	
+	virtual FName GetCategoryName() const override;
+
+	virtual bool NeedsLoadForServer() const override
+	{
+		return false;
+	}
+
+#if WITH_EDITOR
+	virtual void PostEditChangeChainProperty(struct FPropertyChangedChainEvent& PropertyChangedEvent) override;
+#endif //WITH_EDITOR
+
+private:
+	/**
+	 * The image to render for when a requested resource is inaccessible.
+	 * It can be a UTexture or UMaterialInterface or an object implementing the AtlasedTextureInterface.
+	 */
+	UPROPERTY(Config, EditAnywhere, Category = "UI Framework", meta = (AllowPrivateAccess = "true", DisplayThumbnail = "true", AllowedClasses = "/Script/Engine.Texture,/Script/Engine.MaterialInterface,/Script/Engine.SlateTextureAtlasInterface", DisallowedClasses = "/Script/MediaAssets.MediaTexture"))
+	TSoftObjectPtr<UObject> ErrorResource;
+
+	/**
+	 * The image to render while a resource is loading.
+	 * It can be a UTexture or UMaterialInterface or an object implementing the AtlasedTextureInterface.
+	 */
+	UPROPERTY(Config, EditAnywhere, Category = "UI Framework", meta = (AllowPrivateAccess = "true", DisplayThumbnail = "true", AllowedClasses = "/Script/Engine.Texture,/Script/Engine.MaterialInterface,/Script/Engine.SlateTextureAtlasInterface", DisallowedClasses = "/Script/MediaAssets.MediaTexture"))
+	TSoftObjectPtr<UObject> LoadingResource;
+
+	UPROPERTY(Transient)
+	TObjectPtr<UObject> ErrorResourcePtr;
+	UPROPERTY(Transient)
+	TObjectPtr<UObject> LoadingResourcePtr;
+};
diff --git a/Engine/Plugins/Experimental/UIFramework/Source/Public/UIFModule.h b/Engine/Plugins/Experimental/UIFramework/Source/Public/UIFModule.h
--- a/Engine/Plugins/Experimental/UIFramework/Source/Public/UIFModule.h
+++ b/Engine/Plugins/Experimental/UIFramework/Source/Public/UIFModule.h
@@ -34,6 +34,9 @@ public:
 	static void SetPresenterClass(TSubclassOf<UUIFrameworkPresenter> Director);
 	static TSubclassOf<UUIFrameworkPresenter> GetPresenterClass();
 
+public:
+	virtual void StartupModule() override;
+
 private:
 	static void AuthorityDetachWidgetFromParentInternal(UUIFrameworkWidget* Child, bool bTemporary);
 	//static UUIFrameworkWidget* AuthorityRenameRecursive(UUIFrameworkPlayerComponent* ReplicationOwner, UUIFrameworkWidget* Widget, UObject* NewOuter);
diff --git a/Engine/Plugins/Experimental/UIFramework/Source/Public/Widgets/UIFImageBlock.h b/Engine/Plugins/Experimental/UIFramework/Source/Public/Widgets/UIFImageBlock.h
new file mode 100644
--- /dev/null
+++ b/Engine/Plugins/Experimental/UIFramework/Source/Public/Widgets/UIFImageBlock.h
@@ -0,0 +1,116 @@
+// Copyright Epic Games, Inc. All Rights Reserved.
+
+#pragma once
+
+#include "Styling/SlateBrush.h"
+#include "UIFWidget.h"
+
+#include "UIFImageBlock.generated.h"
+
+class UImage;
+class UMaterialInterface;
+struct FStreamableHandle;
+class UTexture2D;
+
+/**
+ * 
+ */
+USTRUCT()
+struct FUIFrameworkImageBlockData
+{
+	GENERATED_BODY()
+
+	/** Tinting applied to the image. */
+	UPROPERTY()
+	FLinearColor Tint = FLinearColor::White;
+
+	/**
+	 * The image to render for this brush, can be a UTexture2D or UMaterialInterface or an object implementing
+	 * the AtlasedTextureInterface.
+	 */
+	UPROPERTY()
+	TSoftObjectPtr<UObject> ResourceObject;
+
+	/** Size of the resource in Slate Units */
+	UPROPERTY()
+	FVector2f DesiredSize = FVector2f::ZeroVector;
+
+	/** How to tile the image. */
+	UPROPERTY()
+	TEnumAsByte<ESlateBrushTileType::Type> Tiling = ESlateBrushTileType::NoTile;
+
+	UPROPERTY()
+	bool bUseTextureSize = false;
+};
+
+/**
+ *
+ */
+UCLASS(DisplayName = "Brush UIFramework")
+class UIFRAMEWORK_API UUIFrameworkImageBlock : public UUIFrameworkWidget
+{
+	GENERATED_BODY()
+
+public:
+	UUIFrameworkImageBlock();
+
+	UFUNCTION(BlueprintCallable, BlueprintAuthorityOnly, Category = "UI Framework")
+	void SetMaterial(TSoftObjectPtr<UMaterialInterface> Material);
+
+	UFUNCTION(BlueprintCallable, Category = "UI Framework")
+	TSoftObjectPtr<UObject> GeResourceObject() const
+	{
+		return Data.ResourceObject;
+	}
+	
+	UFUNCTION(BlueprintCallable, BlueprintAuthorityOnly, Category = "UI Framework")
+	void SetTexture(TSoftObjectPtr<UTexture2D> Texture, bool bUseTextureSize);
+	
+	UFUNCTION(BlueprintCallable, BlueprintAuthorityOnly, Category = "UI Framework")
+	void SetTint(FLinearColor Tint);
+
+	UFUNCTION(BlueprintCallable, Category = "UI Framework")
+	FLinearColor GetTint() const
+	{
+		return Data.Tint;
+	}
+
+	UFUNCTION(BlueprintCallable, BlueprintAuthorityOnly, Category = "UI Framework")
+	void SetDesiredSize(FVector2f DesiredSize);
+
+	UFUNCTION(BlueprintCallable, Category = "UI Framework")
+	FVector2f GetDesiredSize() const
+	{
+		return Data.DesiredSize;
+	}
+	
+	UFUNCTION(BlueprintCallable, BlueprintAuthorityOnly, Category = "UI Framework")
+	void SetTiling(ESlateBrushTileType::Type OverflowPolicy);
+
+	UFUNCTION(BlueprintCallable, Category = "UI Framework")
+	ESlateBrushTileType::Type GetTiling() const
+	{
+		return Data.Tiling;
+	}
+
+	virtual void LocalOnUMGWidgetCreated() override;
+	virtual bool LocalIsReplicationReady() const override;
+
+private:
+	UFUNCTION()
+	void OnRep_Data();
+
+	UObject* ProcessResourceObject();
+
+private:
+	/** Tinting applied to the image. */
+	UPROPERTY(ReplicatedUsing = OnRep_Data)
+	FUIFrameworkImageBlockData Data;
+
+	UPROPERTY()
+	bool bWaitForResourceToBeLoaded = false;
+
+	TSharedPtr<FStreamableHandle> StreamingHandle;
+	int32 StreamingCounter = 0;
+};
+
diff --git a/Engine/Plugins/Experimental/UIFramework/Source/UIFramework.Build.cs b/Engine/Plugins/Experimental/UIFramework/Source/UIFramework.Build.cs
--- a/Engine/Plugins/Experimental/UIFramework/Source/UIFramework.Build.cs
+++ b/Engine/Plugins/Experimental/UIFramework/Source/UIFramework.Build.cs
@@ -11,6 +11,7 @@ public class UIFramework : ModuleRules
 			{
 				"Core",
 				"CoreUObject",
+				"DeveloperSettings",
 				"Engine",
 				"ModelViewViewModel",
 				"SlateCore",
