Fix deprecation warning.
#rnx
#tb none

diff --git a/Engine/Plugins/Experimental/UIFramework/Source/Private/Widgets/UIFImageBlock.cpp b/Engine/Plugins/Experimental/UIFramework/Source/Private/Widgets/UIFImageBlock.cpp
--- a/Engine/Plugins/Experimental/UIFramework/Source/Private/Widgets/UIFImageBlock.cpp
+++ b/Engine/Plugins/Experimental/UIFramework/Source/Private/Widgets/UIFImageBlock.cpp
@@ -88,7 +88,7 @@ void UUIFrameworkImageBlock::LocalOnUMGWidgetCreated()
 {
 	UImage* Image = CastChecked<UImage>(LocalGetUMGWidget());
 
-	FSlateBrush TmpBrush = Image->Brush;
+	FSlateBrush TmpBrush = Image->GetBrush();
 	TmpBrush.ImageSize = FVector2D(Data.DesiredSize);
 	TmpBrush.TintColor = Data.Tint;
 	TmpBrush.Tiling = Data.Tiling;
