VerseUI: Add the color_block widget
#prelfight 63cabe4d6b6da91eda02cbf6

diff --git a/Engine/Plugins/Experimental/UIFramework/Source/Private/Widgets/UIFColorBlock.cpp b/Engine/Plugins/Experimental/UIFramework/Source/Private/Widgets/UIFColorBlock.cpp
new file mode 100644
--- /dev/null
+++ b/Engine/Plugins/Experimental/UIFramework/Source/Private/Widgets/UIFColorBlock.cpp
@@ -0,0 +1,81 @@
+// Copyright Epic Games, Inc. All Rights Reserved.
+
+#include "Widgets/UIFColorBlock.h"
+
+#include "Components/Image.h"
+#include "Engine/AssetManager.h"
+#include "Engine/StreamableManager.h"
+#include "Engine/Texture2D.h"
+#include "Materials/Material.h"
+#include "Materials/MaterialInterface.h"
+
+#include "Net/Core/PushModel/PushModel.h"
+#include "Net/UnrealNetwork.h"
+
+#include UE_INLINE_GENERATED_CPP_BY_NAME(UIFColorBlock)
+
+/**
+ *
+ */
+UUIFrameworkColorBlock::UUIFrameworkColorBlock()
+{
+	WidgetClass = UImage::StaticClass();
+}
+
+void UUIFrameworkColorBlock::GetLifetimeReplicatedProps(TArray<FLifetimeProperty>& OutLifetimeProps) const
+{
+	Super::GetLifetimeReplicatedProps(OutLifetimeProps);
+
+	FDoRepLifetimeParams Params;
+	Params.bIsPushBased = true;
+	DOREPLIFETIME_WITH_PARAMS_FAST(ThisClass, Color, Params);
+	DOREPLIFETIME_WITH_PARAMS_FAST(ThisClass, DesiredSize, Params);
+}
+
+void UUIFrameworkColorBlock::SetColor(FLinearColor InColor)
+{
+	if (Color != InColor)
+	{
+		Color = InColor;
+		MARK_PROPERTY_DIRTY_FROM_NAME(ThisClass, Color, this);
+		ForceNetUpdate();
+	}
+}
+
+void UUIFrameworkColorBlock::SetDesiredSize(FVector2f InDesiredSize)
+{
+	if (DesiredSize != InDesiredSize)
+	{
+		DesiredSize = InDesiredSize;
+		MARK_PROPERTY_DIRTY_FROM_NAME(ThisClass, DesiredSize, this);
+		ForceNetUpdate();
+	}
+}
+
+void UUIFrameworkColorBlock::LocalOnUMGWidgetCreated()
+{
+	UImage* Image = CastChecked<UImage>(LocalGetUMGWidget());
+
+	FSlateBrush TmpBrush = Image->Brush;
+	TmpBrush.ImageSize = FVector2D(DesiredSize);
+	TmpBrush.TintColor = Color;
+
+	Image->SetBrush(TmpBrush);
+}
+
+void UUIFrameworkColorBlock::OnRep_Color()
+{
+	if (LocalGetUMGWidget())
+	{
+		UUIFrameworkColorBlock::LocalOnUMGWidgetCreated();
+	}
+}
+
+void UUIFrameworkColorBlock::OnRep_DesiredSize()
+{
+	if (LocalGetUMGWidget())
+	{
+		UUIFrameworkColorBlock::LocalOnUMGWidgetCreated();
+	}
+}
+
diff --git a/Engine/Plugins/Experimental/UIFramework/Source/Public/Widgets/UIFColorBlock.h b/Engine/Plugins/Experimental/UIFramework/Source/Public/Widgets/UIFColorBlock.h
new file mode 100644
--- /dev/null
+++ b/Engine/Plugins/Experimental/UIFramework/Source/Public/Widgets/UIFColorBlock.h
@@ -0,0 +1,60 @@
+// Copyright Epic Games, Inc. All Rights Reserved.
+
+#pragma once
+
+#include "Styling/SlateBrush.h"
+#include "UIFWidget.h"
+
+#include "UIFColorBlock.generated.h"
+
+class UImage;
+class UMaterialInterface;
+struct FStreamableHandle;
+class UTexture2D;
+
+/**
+ *
+ */
+UCLASS(DisplayName = "Brush UIFramework")
+class UIFRAMEWORK_API UUIFrameworkColorBlock : public UUIFrameworkWidget
+{
+	GENERATED_BODY()
+
+public:
+	UUIFrameworkColorBlock();
+
+	UFUNCTION(BlueprintCallable, BlueprintAuthorityOnly, Category = "UI Framework")
+	void SetColor(FLinearColor Tint);
+
+	UFUNCTION(BlueprintCallable, Category = "UI Framework")
+	FLinearColor GetColor() const
+	{
+		return Color;
+	}
+
+	UFUNCTION(BlueprintCallable, BlueprintAuthorityOnly, Category = "UI Framework")
+	void SetDesiredSize(FVector2f DesiredSize);
+
+	UFUNCTION(BlueprintCallable, Category = "UI Framework")
+	FVector2f GetDesiredSize() const
+	{
+		return DesiredSize;
+	}
+	
+	virtual void LocalOnUMGWidgetCreated() override;
+
+private:
+	UFUNCTION()
+	void OnRep_Color();
+	
+	UFUNCTION()
+	void OnRep_DesiredSize();
+
+private:
+	UPROPERTY(ReplicatedUsing = OnRep_Color)
+	FLinearColor Color;
+	
+	UPROPERTY(ReplicatedUsing = OnRep_DesiredSize)
+	FVector2f DesiredSize;
+};
+
