Fix deprecation warning.
#rnx
#tb none

diff --git a/Engine/Plugins/Experimental/UIFramework/Source/Private/Widgets/UIFColorBlock.cpp b/Engine/Plugins/Experimental/UIFramework/Source/Private/Widgets/UIFColorBlock.cpp
--- a/Engine/Plugins/Experimental/UIFramework/Source/Private/Widgets/UIFColorBlock.cpp
+++ b/Engine/Plugins/Experimental/UIFramework/Source/Private/Widgets/UIFColorBlock.cpp
@@ -56,7 +56,7 @@ void UUIFrameworkColorBlock::LocalOnUMGWidgetCreated()
 {
 	UImage* Image = CastChecked<UImage>(LocalGetUMGWidget());
 
-	FSlateBrush TmpBrush = Image->Brush;
+	FSlateBrush TmpBrush = Image->GetBrush();
 	TmpBrush.ImageSize = FVector2D(DesiredSize);
 	TmpBrush.TintColor = Color;
 
