Prevent laoding of resources on the UIFrameworkLocalSettings CDO that causes build failure with GC assumptions

diff --git a/Engine/Plugins/Experimental/UIFramework/Source/Private/UIFModule.cpp b/Engine/Plugins/Experimental/UIFramework/Source/Private/UIFModule.cpp
--- a/Engine/Plugins/Experimental/UIFramework/Source/Private/UIFModule.cpp
+++ b/Engine/Plugins/Experimental/UIFramework/Source/Private/UIFModule.cpp
@@ -164,7 +164,7 @@ TSubclassOf<UUIFrameworkPresenter> FUIFrameworkModule::GetPresenterClass()
 
 void FUIFrameworkModule::StartupModule()
 {
-	GetMutableDefault<UUIFrameworkLocalSettings>()->LoadResources();
+	//GetMutableDefault<UUIFrameworkLocalSettings>()->LoadResources();
 }
 
 IMPLEMENT_MODULE(FUIFrameworkModule, UIFramework)
\ No newline at end of file
