UIFramework: Add UI Framework local asset to the rootset.

diff --git a/Engine/Plugins/Experimental/UIFramework/Source/Private/UIFLocalSettings.cpp b/Engine/Plugins/Experimental/UIFramework/Source/Private/UIFLocalSettings.cpp
--- a/Engine/Plugins/Experimental/UIFramework/Source/Private/UIFLocalSettings.cpp
+++ b/Engine/Plugins/Experimental/UIFramework/Source/Private/UIFLocalSettings.cpp
@@ -12,7 +12,7 @@ UUIFrameworkLocalSettings::UUIFrameworkLocalSettings()
 	LoadingResource = FSoftObjectPath(TEXT("/UIFramework/Materials/M_UI_DefaultLoading.M_UI_DefaultLoading"));
 }
 
-void UUIFrameworkLocalSettings::LoadResources()
+void UUIFrameworkLocalSettings::LoadResources() const
 {
 	Async(EAsyncExecution::TaskGraphMainThread, []()
 		{
@@ -29,7 +29,16 @@ void UUIFrameworkLocalSettings::LoadResources()
 					{
 						UUIFrameworkLocalSettings* Default = GetMutableDefault<UUIFrameworkLocalSettings>();
 						Default->ErrorResourcePtr = Default->ErrorResource.Get();
+						if (Default->ErrorResourcePtr)
+						{
+							Default->ErrorResourcePtr->AddToRoot();
+						}
+
 						Default->LoadingResourcePtr = Default->LoadingResource.Get();
+						if (Default->LoadingResourcePtr)
+						{
+							Default->LoadingResourcePtr->AddToRoot();
+						}
 					}
 				},
 				FStreamableManager::DefaultAsyncLoadPriority);
diff --git a/Engine/Plugins/Experimental/UIFramework/Source/Private/UIFModule.cpp b/Engine/Plugins/Experimental/UIFramework/Source/Private/UIFModule.cpp
--- a/Engine/Plugins/Experimental/UIFramework/Source/Private/UIFModule.cpp
+++ b/Engine/Plugins/Experimental/UIFramework/Source/Private/UIFModule.cpp
@@ -164,7 +164,7 @@ TSubclassOf<UUIFrameworkPresenter> FUIFrameworkModule::GetPresenterClass()
 
 void FUIFrameworkModule::StartupModule()
 {
-	//GetMutableDefault<UUIFrameworkLocalSettings>()->LoadResources();
+	GetDefault<UUIFrameworkLocalSettings>()->LoadResources();
 }
 
 IMPLEMENT_MODULE(FUIFrameworkModule, UIFramework)
\ No newline at end of file
diff --git a/Engine/Plugins/Experimental/UIFramework/Source/Public/UIFLocalSettings.h b/Engine/Plugins/Experimental/UIFramework/Source/Public/UIFLocalSettings.h
--- a/Engine/Plugins/Experimental/UIFramework/Source/Public/UIFLocalSettings.h
+++ b/Engine/Plugins/Experimental/UIFramework/Source/Public/UIFLocalSettings.h
@@ -29,7 +29,7 @@ public:
 		return LoadingResourcePtr;
 	}
 
-	void LoadResources();
+	void LoadResources() const;
 	
 	virtual FName GetCategoryName() const override;
 
@@ -57,8 +57,7 @@ private:
 	UPROPERTY(Config, EditAnywhere, Category = "UI Framework", meta = (AllowPrivateAccess = "true", DisplayThumbnail = "true", AllowedClasses = "/Script/Engine.Texture,/Script/Engine.MaterialInterface,/Script/Engine.SlateTextureAtlasInterface", DisallowedClasses = "/Script/MediaAssets.MediaTexture"))
 	TSoftObjectPtr<UObject> LoadingResource;
 
-	UPROPERTY(Transient)
+	//~ kept alive with by AddToRoot
 	TObjectPtr<UObject> ErrorResourcePtr;
-	UPROPERTY(Transient)
 	TObjectPtr<UObject> LoadingResourcePtr;
 };
