UIFramework: Change the DisplaName of some UIFWIdget.

diff --git a/Engine/Plugins/Experimental/UIFramework/Source/Public/Widgets/UIFColorBlock.h b/Engine/Plugins/Experimental/UIFramework/Source/Public/Widgets/UIFColorBlock.h
--- a/Engine/Plugins/Experimental/UIFramework/Source/Public/Widgets/UIFColorBlock.h
+++ b/Engine/Plugins/Experimental/UIFramework/Source/Public/Widgets/UIFColorBlock.h
@@ -15,7 +15,7 @@ class UTexture2D;
 /**
  *
  */
-UCLASS(DisplayName = "Brush UIFramework")
+UCLASS(DisplayName = "Color Block UIFramework")
 class UIFRAMEWORK_API UUIFrameworkColorBlock : public UUIFrameworkWidget
 {
 	GENERATED_BODY()
diff --git a/Engine/Plugins/Experimental/UIFramework/Source/Public/Widgets/UIFImageBlock.h b/Engine/Plugins/Experimental/UIFramework/Source/Public/Widgets/UIFImageBlock.h
--- a/Engine/Plugins/Experimental/UIFramework/Source/Public/Widgets/UIFImageBlock.h
+++ b/Engine/Plugins/Experimental/UIFramework/Source/Public/Widgets/UIFImageBlock.h
@@ -46,7 +46,7 @@ struct FUIFrameworkImageBlockData
 /**
  *
  */
-UCLASS(DisplayName = "Brush UIFramework")
+UCLASS(DisplayName = "Image Block UIFramework")
 class UIFRAMEWORK_API UUIFrameworkImageBlock : public UUIFrameworkWidget
 {
 	GENERATED_BODY()
