UIFramework: Only load the asset on client when the component is initialized.

diff --git a/Engine/Plugins/Experimental/UIFramework/Source/Private/UIFLocalSettings.cpp b/Engine/Plugins/Experimental/UIFramework/Source/Private/UIFLocalSettings.cpp
--- a/Engine/Plugins/Experimental/UIFramework/Source/Private/UIFLocalSettings.cpp
+++ b/Engine/Plugins/Experimental/UIFramework/Source/Private/UIFLocalSettings.cpp
@@ -14,35 +14,46 @@ UUIFrameworkLocalSettings::UUIFrameworkLocalSettings()
 
 void UUIFrameworkLocalSettings::LoadResources() const
 {
-	Async(EAsyncExecution::TaskGraphMainThread, []()
+	if (bResourceLoaded)
+	{
+		return;
+	}
+	bResourceLoaded = true;
+
+	const UUIFrameworkLocalSettings* Settings = GetDefault<UUIFrameworkLocalSettings>();
+	TArray<FSoftObjectPath> ObjectsToLoad;
+	ObjectsToLoad.Reserve(2);
+	ObjectsToLoad.Add(Settings->ErrorResource.ToSoftObjectPath());
+	ObjectsToLoad.Add(Settings->LoadingResource.ToSoftObjectPath());
+	UAssetManager::GetStreamableManager().RequestAsyncLoad(
+		ObjectsToLoad,
+		[]()
 		{
-			const UUIFrameworkLocalSettings* Settings = GetDefault<UUIFrameworkLocalSettings>();
-			TArray<FSoftObjectPath> ObjectsToLoad;
-			ObjectsToLoad.Reserve(2);
-			ObjectsToLoad.Add(Settings->ErrorResource.ToSoftObjectPath());
-			ObjectsToLoad.Add(Settings->LoadingResource.ToSoftObjectPath());
-			UAssetManager::GetStreamableManager().RequestAsyncLoad(
-				ObjectsToLoad,
-				[]()
+			if (UObjectInitialized() && !IsEngineExitRequested())
+			{
+				UUIFrameworkLocalSettings* Default = GetMutableDefault<UUIFrameworkLocalSettings>();
+				if (Default->ErrorResourcePtr)
+				{
+					Default->ErrorResourcePtr->RemoveFromRoot();
+				}
+				Default->ErrorResourcePtr = Default->ErrorResource.Get();
+				if (Default->ErrorResourcePtr)
+				{
+					Default->ErrorResourcePtr->AddToRoot();
+				}
+
+				if (Default->LoadingResourcePtr)
+				{
+					Default->LoadingResourcePtr->RemoveFromRoot();
+				}
+				Default->LoadingResourcePtr = Default->LoadingResource.Get();
+				if (Default->LoadingResourcePtr)
 				{
-					if (UObjectInitialized() && !IsEngineExitRequested())
-					{
-						UUIFrameworkLocalSettings* Default = GetMutableDefault<UUIFrameworkLocalSettings>();
-						Default->ErrorResourcePtr = Default->ErrorResource.Get();
-						if (Default->ErrorResourcePtr)
-						{
-							Default->ErrorResourcePtr->AddToRoot();
-						}
-
-						Default->LoadingResourcePtr = Default->LoadingResource.Get();
-						if (Default->LoadingResourcePtr)
-						{
-							Default->LoadingResourcePtr->AddToRoot();
-						}
-					}
-				},
-				FStreamableManager::DefaultAsyncLoadPriority);
-		});
+					Default->LoadingResourcePtr->AddToRoot();
+				}
+			}
+		},
+		FStreamableManager::DefaultAsyncLoadPriority);
 }
 
 
@@ -55,6 +66,8 @@ FName UUIFrameworkLocalSettings::GetCategoryName() const
 void UUIFrameworkLocalSettings::PostEditChangeChainProperty(struct FPropertyChangedChainEvent& PropertyChangedEvent)
 {
 	Super::PostEditChangeChainProperty(PropertyChangedEvent);
+
+	bResourceLoaded = false;
 	LoadResources();
 }
 #endif
\ No newline at end of file
diff --git a/Engine/Plugins/Experimental/UIFramework/Source/Private/UIFModule.cpp b/Engine/Plugins/Experimental/UIFramework/Source/Private/UIFModule.cpp
--- a/Engine/Plugins/Experimental/UIFramework/Source/Private/UIFModule.cpp
+++ b/Engine/Plugins/Experimental/UIFramework/Source/Private/UIFModule.cpp
@@ -3,7 +3,6 @@
 #include "UIFModule.h"
 
 #include "Modules/ModuleManager.h"
-#include "UIFLocalSettings.h"
 #include "UIFPresenter.h"
 #include "UIFWidget.h"
 
@@ -161,9 +160,4 @@ TSubclassOf<UUIFrameworkPresenter> FUIFrameworkModule::GetPresenterClass()
 	return UE::UIFramework::Private::Director.Get() ? UE::UIFramework::Private::Director : TSubclassOf<UUIFrameworkPresenter>(UUIFrameworkGameViewportPresenter::StaticClass());
 }
 
-void FUIFrameworkModule::StartupModule()
-{
-	GetDefault<UUIFrameworkLocalSettings>()->LoadResources();
-}
-
 IMPLEMENT_MODULE(FUIFrameworkModule, UIFramework)
diff --git a/Engine/Plugins/Experimental/UIFramework/Source/Private/UIFPlayerComponent.cpp b/Engine/Plugins/Experimental/UIFramework/Source/Private/UIFPlayerComponent.cpp
--- a/Engine/Plugins/Experimental/UIFramework/Source/Private/UIFPlayerComponent.cpp
+++ b/Engine/Plugins/Experimental/UIFramework/Source/Private/UIFPlayerComponent.cpp
@@ -6,6 +6,7 @@
 #include "Templates/NonNullPointer.h"
 #include "UIFModule.h"
 #include "Types/UIFWidgetOwner.h"
+#include "UIFLocalSettings.h"
 #include "UIFPresenter.h"
 #include "UIFWidget.h"
 
@@ -106,6 +107,10 @@ void UUIFrameworkPlayerComponent::InitializeComponent()
 	{
 		WidgetTree.AuthorityAddAllWidgetsFromActorChannel();
 	}
+	else
+	{
+		GetDefault<UUIFrameworkLocalSettings>()->LoadResources();
+	}
 }
 
 void UUIFrameworkPlayerComponent::UninitializeComponent()
diff --git a/Engine/Plugins/Experimental/UIFramework/Source/Public/UIFLocalSettings.h b/Engine/Plugins/Experimental/UIFramework/Source/Public/UIFLocalSettings.h
--- a/Engine/Plugins/Experimental/UIFramework/Source/Public/UIFLocalSettings.h
+++ b/Engine/Plugins/Experimental/UIFramework/Source/Public/UIFLocalSettings.h
@@ -59,6 +59,7 @@ private:
 	//~ kept alive with by AddToRoot
 	TObjectPtr<UObject> ErrorResourcePtr;
 	TObjectPtr<UObject> LoadingResourcePtr;
+	mutable bool bResourceLoaded = false;
 };
 
 #if UE_ENABLE_INCLUDE_ORDER_DEPRECATED_IN_5_2
diff --git a/Engine/Plugins/Experimental/UIFramework/Source/Public/UIFModule.h b/Engine/Plugins/Experimental/UIFramework/Source/Public/UIFModule.h
--- a/Engine/Plugins/Experimental/UIFramework/Source/Public/UIFModule.h
+++ b/Engine/Plugins/Experimental/UIFramework/Source/Public/UIFModule.h
@@ -34,9 +34,6 @@ public:
 	static void SetPresenterClass(TSubclassOf<UUIFrameworkPresenter> Director);
 	static TSubclassOf<UUIFrameworkPresenter> GetPresenterClass();
 
-public:
-	virtual void StartupModule() override;
-
 private:
 	static void AuthorityDetachWidgetFromParentInternal(UUIFrameworkWidget* Child, bool bTemporary);
 	//static UUIFrameworkWidget* AuthorityRenameRecursive(UUIFrameworkPlayerComponent* ReplicationOwner, UUIFrameworkWidget* Widget, UObject* NewOuter);
