UIFramework: The UIFButton use a transparent border by default.

diff --git a/Engine/Plugins/Experimental/UIFramework/Source/Private/Widgets/UIFButton.cpp b/Engine/Plugins/Experimental/UIFramework/Source/Private/Widgets/UIFButton.cpp
--- a/Engine/Plugins/Experimental/UIFramework/Source/Private/Widgets/UIFButton.cpp
+++ b/Engine/Plugins/Experimental/UIFramework/Source/Private/Widgets/UIFButton.cpp
@@ -18,7 +18,7 @@
  */
 UUIFrameworkButton::UUIFrameworkButton()
 {
-	WidgetClass = UButton::StaticClass();
+	WidgetClass = FSoftObjectPath(TEXT("/UIFramework/Widgets/WBP_UIF_Button.WBP_UIF_Button"));
 }
 
 void UUIFrameworkButton::GetLifetimeReplicatedProps(TArray<FLifetimeProperty>& OutLifetimeProps) const
