Fix for UUIFrameworkButton failing to load its UMG widget class (because it wasn't cooked). Temp fix not needing an asset to load.

diff --git a/Engine/Plugins/Experimental/UIFramework/Source/Private/UIFPlayerComponent.cpp b/Engine/Plugins/Experimental/UIFramework/Source/Private/UIFPlayerComponent.cpp
--- a/Engine/Plugins/Experimental/UIFramework/Source/Private/UIFPlayerComponent.cpp
+++ b/Engine/Plugins/Experimental/UIFramework/Source/Private/UIFPlayerComponent.cpp
@@ -386,7 +386,8 @@ void UUIFrameworkPlayerComponent::LocalOnClassLoaded(TSoftClassPtr<UWidget> Widg
 		}
 		else
 		{
-			ensureMsgf(false, TEXT("Load request failed"));
+			UE_LOG(LogUIFramework, Error, TEXT("Failed to load widget class %s."), *WidgetClass.ToString());
+			ensure(false);
 		}
 	}
 	else
diff --git a/Engine/Plugins/Experimental/UIFramework/Source/Private/Widgets/UIFButton.cpp b/Engine/Plugins/Experimental/UIFramework/Source/Private/Widgets/UIFButton.cpp
--- a/Engine/Plugins/Experimental/UIFramework/Source/Private/Widgets/UIFButton.cpp
+++ b/Engine/Plugins/Experimental/UIFramework/Source/Private/Widgets/UIFButton.cpp
@@ -8,17 +8,35 @@
 #include "Components/Button.h"
 #include "Components/ButtonSlot.h"
 
+#include "Misc/RedirectCollector.h"
 #include "Net/UnrealNetwork.h"
 
 #include UE_INLINE_GENERATED_CPP_BY_NAME(UIFButton)
 
 
+UUIFrameworkButtonWidget::UUIFrameworkButtonWidget()
+{
+	WidgetStyle.NormalPadding = FMargin(0.0f);
+	WidgetStyle.Normal.DrawAs = ESlateBrushDrawType::NoDrawType;
+	WidgetStyle.Hovered.DrawAs = ESlateBrushDrawType::NoDrawType;
+	WidgetStyle.Pressed.DrawAs = ESlateBrushDrawType::NoDrawType;
+}
+
 /**
  * 
  */
 UUIFrameworkButton::UUIFrameworkButton()
 {
-	WidgetClass = FSoftObjectPath(TEXT("/UIFramework/Widgets/WBP_UIF_Button.WBP_UIF_Button"));
+// TODO: FIXME: Loading this asset fails. Temporary fixed by using UUIFrameworkButtonWidget::StaticClass()
+// 	            Does this actually need to be an asset meant to be customizable or can it be set like above?
+
+//	WidgetClass = FSoftObjectPath(TEXT("/UIFramework/Widgets/WBP_UIF_Button.WBP_UIF_Button"));
+//#if WITH_EDITOR
+//	// We need to ensure that this asset is cooked
+//	GRedirectCollector.OnSoftObjectPathLoaded(WidgetClass.ToSoftObjectPath(), nullptr);
+//#endif
+
+	WidgetClass = UUIFrameworkButtonWidget::StaticClass();
 }
 
 void UUIFrameworkButton::GetLifetimeReplicatedProps(TArray<FLifetimeProperty>& OutLifetimeProps) const
diff --git a/Engine/Plugins/Experimental/UIFramework/Source/Public/Widgets/UIFButton.h b/Engine/Plugins/Experimental/UIFramework/Source/Public/Widgets/UIFButton.h
--- a/Engine/Plugins/Experimental/UIFramework/Source/Public/Widgets/UIFButton.h
+++ b/Engine/Plugins/Experimental/UIFramework/Source/Public/Widgets/UIFButton.h
@@ -2,6 +2,7 @@
 
 #pragma once
 
+#include "Components/Button.h"
 #include "Types/UIFSlotBase.h"
 #include "Types/UIFEvents.h"
 #include "UIFWidget.h"
@@ -54,3 +55,12 @@ private:
 	UPROPERTY(/*ExposeOnSpawn, */ReplicatedUsing = OnRep_Slot)
 	FUIFrameworkSimpleSlot Slot;
 };
+
+UCLASS()
+class UIFRAMEWORK_API UUIFrameworkButtonWidget : public UButton
+{
+	GENERATED_BODY()
+
+public:
+	UUIFrameworkButtonWidget();
+};
