Fix compile errors with deprecated property access

diff --git a/Engine/Plugins/Experimental/UIFramework/Source/Private/Widgets/UIFButton.cpp b/Engine/Plugins/Experimental/UIFramework/Source/Private/Widgets/UIFButton.cpp
--- a/Engine/Plugins/Experimental/UIFramework/Source/Private/Widgets/UIFButton.cpp
+++ b/Engine/Plugins/Experimental/UIFramework/Source/Private/Widgets/UIFButton.cpp
@@ -16,10 +16,12 @@
 
 UUIFrameworkButtonWidget::UUIFrameworkButtonWidget()
 {
-	WidgetStyle.NormalPadding = FMargin(0.0f);
-	WidgetStyle.Normal.DrawAs = ESlateBrushDrawType::NoDrawType;
-	WidgetStyle.Hovered.DrawAs = ESlateBrushDrawType::NoDrawType;
-	WidgetStyle.Pressed.DrawAs = ESlateBrushDrawType::NoDrawType;
+	FButtonStyle TempStyle = GetStyle();
+	TempStyle.NormalPadding = FMargin(0.0f);
+	TempStyle.Normal.DrawAs = ESlateBrushDrawType::NoDrawType;
+	TempStyle.Hovered.DrawAs = ESlateBrushDrawType::NoDrawType;
+	TempStyle.Pressed.DrawAs = ESlateBrushDrawType::NoDrawType;
+	SetStyle(TempStyle);
 }
 
 /**
