UIFramework: Do not add the widget to the replication list when the player component uses the registered sub objec list.

diff --git a/Engine/Plugins/Experimental/UIFramework/Source/Private/UIFPlayerComponent.cpp b/Engine/Plugins/Experimental/UIFramework/Source/Private/UIFPlayerComponent.cpp
--- a/Engine/Plugins/Experimental/UIFramework/Source/Private/UIFPlayerComponent.cpp
+++ b/Engine/Plugins/Experimental/UIFramework/Source/Private/UIFPlayerComponent.cpp
@@ -146,9 +146,12 @@ void UUIFrameworkPlayerComponent::GetLifetimeReplicatedProps(TArray< FLifetimePr
 
 bool UUIFrameworkPlayerComponent::ReplicateSubobjects(UActorChannel* Channel, class FOutBunch* Bunch, FReplicationFlags* RepFlags)
 {
-	bool WroteSomething = Super::ReplicateSubobjects(Channel, Bunch, RepFlags);
-	WroteSomething |= WidgetTree.ReplicateSubWidgets(Channel, Bunch, RepFlags);
-	return WroteSomething;
+	bool bWroteSomething = Super::ReplicateSubobjects(Channel, Bunch, RepFlags);
+	if (!bReplicateUsingRegisteredSubObjectList)
+	{
+		bWroteSomething |= WidgetTree.ReplicateSubWidgets(Channel, Bunch, RepFlags);
+	}
+	return bWroteSomething;
 }
 
 void UUIFrameworkPlayerComponent::AddWidget(FUIFrameworkGameLayerSlot InEntry)
