Work on Localization: Implement Message replication and integrate it on Verse UI Framework

diff --git a/Engine/Plugins/Experimental/UIFramework/Source/Private/Widgets/UIFSimpleButton.cpp b/Engine/Plugins/Experimental/UIFramework/Source/Private/Widgets/UIFSimpleButton.cpp
--- a/Engine/Plugins/Experimental/UIFramework/Source/Private/Widgets/UIFSimpleButton.cpp
+++ b/Engine/Plugins/Experimental/UIFramework/Source/Private/Widgets/UIFSimpleButton.cpp
@@ -3,6 +3,8 @@
 #include "Widgets/UIFSimpleButton.h"
 
 #include "Blueprint/UserWidget.h"
+#include "IVerseModule.h"
+#include "Localization/VerseLocalizationProcessor.h"
 #include "MVVMSubsystem.h"
 #include "Net/Core/PushModel/PushModel.h"
 #include "Net/UnrealNetwork.h"
@@ -27,7 +29,7 @@ void UUIFrameworkSimpleButton::GetLifetimeReplicatedProps(TArray<FLifetimeProper
 
 	FDoRepLifetimeParams Params;
 	Params.bIsPushBased = true;
-	DOREPLIFETIME_WITH_PARAMS_FAST(ThisClass, Text, Params);
+	DOREPLIFETIME_WITH_PARAMS_FAST(ThisClass, Message, Params);
 }
 
 
@@ -44,17 +46,21 @@ void UUIFrameworkSimpleButton::LocalOnUMGWidgetCreated()
 }
 
 
-void UUIFrameworkSimpleButton::SetText(FText InText)
+void UUIFrameworkSimpleButton::SetMessage(FVerseReplicationMessage&& InMessage)
 {
-	Text = InText;
-	MARK_PROPERTY_DIRTY_FROM_NAME(ThisClass, Text, this);
+	Message = MoveTemp(InMessage);
+	MARK_PROPERTY_DIRTY_FROM_NAME(ThisClass, Message, this);
 	ForceNetUpdate();
-	UE_MVVM_BROADCAST_FIELD_VALUE_CHANGED(Text);
 }
 
 
-void UUIFrameworkSimpleButton::OnRep_Text()
+void UUIFrameworkSimpleButton::OnRep_Message()
 {
+	IVerseModule& VerseModule = IVerseModule::Get();
+	verse::FLocalizationProcessor& Processor = VerseModule.GetLocalizationProcessor();
+	TNonNullPtr<verse::message> VerseMessage = Processor.FromReplicationMessage(Message,GetWorld());
+
+	Text = Processor.Localize(VerseMessage);
 	UE_MVVM_BROADCAST_FIELD_VALUE_CHANGED(Text);
 }
 
diff --git a/Engine/Plugins/Experimental/UIFramework/Source/Private/Widgets/UIFTextBlock.cpp b/Engine/Plugins/Experimental/UIFramework/Source/Private/Widgets/UIFTextBlock.cpp
--- a/Engine/Plugins/Experimental/UIFramework/Source/Private/Widgets/UIFTextBlock.cpp
+++ b/Engine/Plugins/Experimental/UIFramework/Source/Private/Widgets/UIFTextBlock.cpp
@@ -3,6 +3,8 @@
 #include "Widgets/UIFTextBlock.h"
 
 #include "Components/TextBlock.h"
+#include "IVerseModule.h"
+#include "Localization/VerseLocalizationProcessor.h"
 #include "Net/Core/PushModel/PushModel.h"
 #include "Net/UnrealNetwork.h"
 
@@ -24,16 +26,16 @@ void UUIFrameworkTextBase::GetLifetimeReplicatedProps(TArray<FLifetimeProperty>&
 
 	FDoRepLifetimeParams Params;
 	Params.bIsPushBased = true;
-	DOREPLIFETIME_WITH_PARAMS_FAST(ThisClass, Text, Params);
+	DOREPLIFETIME_WITH_PARAMS_FAST(ThisClass, Message, Params);
 	DOREPLIFETIME_WITH_PARAMS_FAST(ThisClass, TextColor, Params);
 	DOREPLIFETIME_WITH_PARAMS_FAST(ThisClass, Justification, Params);
 	DOREPLIFETIME_WITH_PARAMS_FAST(ThisClass, OverflowPolicy, Params);
 }
 
-void UUIFrameworkTextBase::SetText(FText InText)
+void UUIFrameworkTextBase::SetMessage(FVerseReplicationMessage&& InMessage)
 {
-	Text = InText;
-	MARK_PROPERTY_DIRTY_FROM_NAME(ThisClass, Text, this);
+	Message = MoveTemp(InMessage);
+	MARK_PROPERTY_DIRTY_FROM_NAME(ThisClass, Message, this);
 	ForceNetUpdate();
 }
 
@@ -75,8 +77,13 @@ void UUIFrameworkTextBase::LocalOnUMGWidgetCreated()
 	SetOverflowPolicyToWidget(OverflowPolicy);
 }
 
-void UUIFrameworkTextBase::OnRep_Text()
+void UUIFrameworkTextBase::OnRep_Message()
 {
+	IVerseModule& VerseModule = IVerseModule::Get();
+	verse::FLocalizationProcessor& Processor = VerseModule.GetLocalizationProcessor();
+	TNonNullPtr<verse::message> VerseMessage = Processor.FromReplicationMessage(Message, GetWorld());
+	Text = Processor.Localize(VerseMessage);
+
 	if (LocalGetUMGWidget())
 	{
 		SetTextToWidget(Text);
diff --git a/Engine/Plugins/Experimental/UIFramework/Source/Public/Widgets/UIFSimpleButton.h b/Engine/Plugins/Experimental/UIFramework/Source/Public/Widgets/UIFSimpleButton.h
--- a/Engine/Plugins/Experimental/UIFramework/Source/Public/Widgets/UIFSimpleButton.h
+++ b/Engine/Plugins/Experimental/UIFramework/Source/Public/Widgets/UIFSimpleButton.h
@@ -5,6 +5,8 @@
 #include "Types/UIFEvents.h"
 #include "UIFWidget.h"
 
+#include "Localization/VerseReplicationMessage.h"
+
 #include "UIFSimpleButton.generated.h"
 
 struct FMVVMEventField;
@@ -21,21 +23,24 @@ public:
 	UUIFrameworkSimpleButton();
 
 private:
-	UPROPERTY(ReplicatedUsing=OnRep_Text, EditAnywhere, BlueprintReadOnly, Getter, Setter, FieldNotify, Category="UI Framework", meta = (AllowPrivateAccess = "true"))
+	UPROPERTY(BlueprintReadOnly, Getter, FieldNotify, Category="UI Framework", meta = (AllowPrivateAccess = "true"))
 	FText Text;
 
+	UPROPERTY(ReplicatedUsing=OnRep_Message, FieldNotify)
+	FVerseReplicationMessage Message;
+
 	UPROPERTY(BlueprintReadOnly, Getter, FieldNotify, Category = "UI Framework", meta = (DisallowedViewAccess = "true", AllowPrivateAccess = "true"))
 	FUIFrameworkClickEventArgument ClickEvent;
 
 public:
-	UFUNCTION(BlueprintCallable, BlueprintAuthorityOnly, Category = "UI Framework")
-	void SetText(FText Value);
 
 	FText GetText() const
 	{
 		return Text;
 	}
 
+	void SetMessage(FVerseReplicationMessage&& InMessage);
+
 	const FUIFrameworkClickEventArgument& GetClickEvent() const
 	{
 		return ClickEvent;
@@ -47,8 +52,9 @@ public:
 	void OnClick(FMVVMEventField Field);
 
 private:
+
 	UFUNCTION()
-	void OnRep_Text();
+	void OnRep_Message();
 
 	UFUNCTION(Server, Reliable)
 	void ServerClick(APlayerController* PlayerController);
diff --git a/Engine/Plugins/Experimental/UIFramework/Source/Public/Widgets/UIFTextBlock.h b/Engine/Plugins/Experimental/UIFramework/Source/Public/Widgets/UIFTextBlock.h
--- a/Engine/Plugins/Experimental/UIFramework/Source/Public/Widgets/UIFTextBlock.h
+++ b/Engine/Plugins/Experimental/UIFramework/Source/Public/Widgets/UIFTextBlock.h
@@ -5,6 +5,8 @@
 #include "Styling/SlateTypes.h"
 #include "UIFWidget.h"
 
+#include "Localization/VerseReplicationMessage.h"
+
 #include "UIFTextBlock.generated.h"
 
 namespace ETextJustify { enum Type : int; }
@@ -22,8 +24,7 @@ class UIFRAMEWORK_API UUIFrameworkTextBase : public UUIFrameworkWidget
 public:
 	UUIFrameworkTextBase();
 
-	UFUNCTION(BlueprintCallable, BlueprintAuthorityOnly, Category = "UI Framework")
-	void SetText(FText Text);
+	void SetMessage(FVerseReplicationMessage&& InMessage);
 
 	UFUNCTION(BlueprintCallable, Category = "UI Framework")
 	FText GetText() const
@@ -62,7 +63,7 @@ public:
 
 private:
 	UFUNCTION()
-	void OnRep_Text();
+	void OnRep_Message();
 	
 	UFUNCTION()
 	void OnRep_TextColor();
@@ -88,9 +89,12 @@ protected:
 	}
 
 private:
-	UPROPERTY(/*ExposeOnSpawn, */ReplicatedUsing = OnRep_Text)
+	UPROPERTY()
 	FText Text;
 
+	UPROPERTY(ReplicatedUsing = OnRep_Message)
+	FVerseReplicationMessage Message;
+
 	UPROPERTY(ReplicatedUsing = OnRep_TextColor)
 	FLinearColor TextColor = FLinearColor::Black;
 
diff --git a/Engine/Plugins/Experimental/UIFramework/Source/UIFramework.Build.cs b/Engine/Plugins/Experimental/UIFramework/Source/UIFramework.Build.cs
--- a/Engine/Plugins/Experimental/UIFramework/Source/UIFramework.Build.cs
+++ b/Engine/Plugins/Experimental/UIFramework/Source/UIFramework.Build.cs
@@ -18,6 +18,7 @@ public class UIFramework : ModuleRules
 				"SlateCore",
 				"Slate",
 				"UMG",
+				"Verse"
 			}
 		);
 
