UIFramework:  Fix badly initialized structured.

diff --git a/Engine/Plugins/Experimental/UIFramework/Source/Public/Types/UIFSlotBase.h b/Engine/Plugins/Experimental/UIFramework/Source/Public/Types/UIFSlotBase.h
--- a/Engine/Plugins/Experimental/UIFramework/Source/Public/Types/UIFSlotBase.h
+++ b/Engine/Plugins/Experimental/UIFramework/Source/Public/Types/UIFSlotBase.h
@@ -76,10 +76,10 @@ struct FUIFrameworkSimpleSlot : public FUIFrameworkSlotBase
 	FMargin Padding;
 
 	UPROPERTY(BlueprintReadWrite, Category = "UI Framework")
-	TEnumAsByte<EHorizontalAlignment> HorizontalAlignment;
+	TEnumAsByte<EHorizontalAlignment> HorizontalAlignment = EHorizontalAlignment::HAlign_Fill;
 
 	UPROPERTY(BlueprintReadWrite, Category = "UI Framework")
-	TEnumAsByte<EVerticalAlignment> VerticalAlignment;
+	TEnumAsByte<EVerticalAlignment> VerticalAlignment = EVerticalAlignment::VAlign_Fill;
 };
 
 #if UE_ENABLE_INCLUDE_ORDER_DEPRECATED_IN_5_2
diff --git a/Engine/Plugins/Experimental/UIFramework/Source/Public/Widgets/UIFOverlay.h b/Engine/Plugins/Experimental/UIFramework/Source/Public/Widgets/UIFOverlay.h
--- a/Engine/Plugins/Experimental/UIFramework/Source/Public/Widgets/UIFOverlay.h
+++ b/Engine/Plugins/Experimental/UIFramework/Source/Public/Widgets/UIFOverlay.h
@@ -38,7 +38,7 @@ struct FUIFrameworkOverlaySlot : public FUIFrameworkSlotBase
 private:
 	/** Index in the array the Slot is. The position in the array can change when replicated. */
 	UPROPERTY()
-	int32 Index;
+	int32 Index = INDEX_NONE;
 };
 
 
diff --git a/Engine/Plugins/Experimental/UIFramework/Source/Public/Widgets/UIFStackBox.h b/Engine/Plugins/Experimental/UIFramework/Source/Public/Widgets/UIFStackBox.h
--- a/Engine/Plugins/Experimental/UIFramework/Source/Public/Widgets/UIFStackBox.h
+++ b/Engine/Plugins/Experimental/UIFramework/Source/Public/Widgets/UIFStackBox.h
@@ -42,7 +42,7 @@ struct FUIFrameworkStackBoxSlot : public FUIFrameworkSlotBase
 private:
 	/** Index in the array the Slot is. The position in the array can change when replicated. */
 	UPROPERTY()
-	int32 Index;
+	int32 Index = INDEX_NONE;
 };
 
 
