Open on a UObject Rename and if we abort during a transaction Rename back to the outer we were originally

diff --git a/Engine/Plugins/Experimental/UIFramework/Source/Private/Types/UIFWidgetTree.cpp b/Engine/Plugins/Experimental/UIFramework/Source/Private/Types/UIFWidgetTree.cpp
--- a/Engine/Plugins/Experimental/UIFramework/Source/Private/Types/UIFWidgetTree.cpp
+++ b/Engine/Plugins/Experimental/UIFramework/Source/Private/Types/UIFWidgetTree.cpp
@@ -219,7 +219,16 @@ void FUIFrameworkWidgetTree::AuthorityAddChildInternal(UUIFrameworkWidget* Paren
 			//If the outer is the transient package, then there are no replication owner yetand it safe to rename it with the correct new outer.
 			//If the Outer is not the transient, then the widget got replicated with another player.There are no "reset" and we should duplicate and delete the object.
 			//For now only do the rename.It works but it is not the best.
-			Child->Rename(nullptr, ReplicatedOwner);
+
+			UObject* OldOuter = Child->GetOuter();
+
+			UE_AUTORTFM_OPEN({
+				Child->Rename(nullptr, ReplicatedOwner);
+			});
+
+			UE_AUTORTFM_OPENABORT({
+				Child->Rename(nullptr, OldOuter);
+			});
 		}
 	}
 
