Update the textblock localization message if the replicated data comes before the widget is created.

diff --git a/Engine/Plugins/Experimental/UIFramework/Source/Private/Widgets/UIFTextBlock.cpp b/Engine/Plugins/Experimental/UIFramework/Source/Private/Widgets/UIFTextBlock.cpp
--- a/Engine/Plugins/Experimental/UIFramework/Source/Private/Widgets/UIFTextBlock.cpp
+++ b/Engine/Plugins/Experimental/UIFramework/Source/Private/Widgets/UIFTextBlock.cpp
@@ -86,6 +86,7 @@ void UUIFrameworkTextBase::SetOverflowPolicy(ETextOverflowPolicy InOverflowPolic
 
 void UUIFrameworkTextBase::LocalOnUMGWidgetCreated()
 {
+	SetText(Message);
 	SetTextToWidget(Text);
 	SetTextColorToWidget(TextColor);
 	SetJustificationToWidget(Justification);
