UIFramework: Add viewmodel to UIF UserWidget.

diff --git a/Engine/Plugins/Experimental/UIFramework/Source/Private/Widgets/UIFUserWidget.cpp b/Engine/Plugins/Experimental/UIFramework/Source/Private/Widgets/UIFUserWidget.cpp
--- a/Engine/Plugins/Experimental/UIFramework/Source/Private/Widgets/UIFUserWidget.cpp
+++ b/Engine/Plugins/Experimental/UIFramework/Source/Private/Widgets/UIFUserWidget.cpp
@@ -6,6 +6,8 @@
 #include "UIFModule.h"
 
 #include "Blueprint/UserWidget.h"
+#include "MVVMSubsystem.h"
+#include "View/MVVMView.h"
 
 #include "Net/UnrealNetwork.h"
 
@@ -15,12 +17,12 @@
 /**
  *
  */
-void FUIFrameworkUserWidgetSlotList::PostReplicatedChange(const TArrayView<int32>& ChangedIndices, int32 FinalSize)
+void FUIFrameworkUserWidgetNamedSlotList::PostReplicatedChange(const TArrayView<int32>& ChangedIndices, int32 FinalSize)
 {
 	check(Owner);
 	for (int32 Index : ChangedIndices)
 	{
-		FUIFrameworkUserWidgetSlot& Slot = Slots[Index];
+		FUIFrameworkUserWidgetNamedSlot& Slot = Slots[Index];
 		if (Slot.LocalIsAquiredWidgetValid())
 		{
 			// Remove and add the widget again... 
@@ -31,21 +33,28 @@ void FUIFrameworkUserWidgetSlotList::PostReplicatedChange(const TArrayView<int32
 	}
 }
 
-bool FUIFrameworkUserWidgetSlotList::NetDeltaSerialize(FNetDeltaSerializeInfo& DeltaParms)
+bool FUIFrameworkUserWidgetNamedSlotList::NetDeltaSerialize(FNetDeltaSerializeInfo& DeltaParms)
 {
-	return FFastArraySerializer::FastArrayDeltaSerialize<FUIFrameworkUserWidgetSlot, FUIFrameworkUserWidgetSlotList>(Slots, DeltaParms, *this);
+	return FFastArraySerializer::FastArrayDeltaSerialize<FUIFrameworkUserWidgetNamedSlot, FUIFrameworkUserWidgetNamedSlotList>(Slots, DeltaParms, *this);
 }
 
-void FUIFrameworkUserWidgetSlotList::AddEntry(FUIFrameworkUserWidgetSlot Entry)
+void FUIFrameworkUserWidgetNamedSlotList::AuthorityAddEntry(FUIFrameworkUserWidgetNamedSlot Entry)
 {
-	FUIFrameworkUserWidgetSlot& NewEntry = Slots.Add_GetRef(MoveTemp(Entry));
+	// Make sure there is only one entry of that name.
+	int32 RemoveCount = Slots.RemoveAllSwap([SlotName = Entry.SlotName](const FUIFrameworkUserWidgetNamedSlot& Entry) { return Entry.SlotName == SlotName; });
+	if (RemoveCount > 0)
+	{
+		MarkArrayDirty();
+	}
+
+	FUIFrameworkUserWidgetNamedSlot& NewEntry = Slots.Add_GetRef(MoveTemp(Entry));
 	MarkItemDirty(NewEntry);
 }
 
-bool FUIFrameworkUserWidgetSlotList::RemoveEntry(UUIFrameworkWidget* Widget)
+bool FUIFrameworkUserWidgetNamedSlotList::AuthorityRemoveEntry(UUIFrameworkWidget* Widget)
 {
 	check(Widget);
-	const int32 Index = Slots.IndexOfByPredicate([Widget](const FUIFrameworkUserWidgetSlot& Entry) { return Entry.AuthorityGetWidget() == Widget; });
+	const int32 Index = Slots.IndexOfByPredicate([Widget](const FUIFrameworkUserWidgetNamedSlot& Entry) { return Entry.AuthorityGetWidget() == Widget; });
 	if (Index != INDEX_NONE)
 	{
 		Slots.RemoveAt(Index);
@@ -54,14 +63,19 @@ bool FUIFrameworkUserWidgetSlotList::RemoveEntry(UUIFrameworkWidget* Widget)
 	return Index != INDEX_NONE;
 }
 
-FUIFrameworkUserWidgetSlot* FUIFrameworkUserWidgetSlotList::FindEntry(FUIFrameworkWidgetId WidgetId)
+FUIFrameworkUserWidgetNamedSlot* FUIFrameworkUserWidgetNamedSlotList::FindEntry(FUIFrameworkWidgetId WidgetId)
+{
+	return Slots.FindByPredicate([WidgetId](const FUIFrameworkUserWidgetNamedSlot& Entry) { return Entry.GetWidgetId() == WidgetId; });
+}
+
+const FUIFrameworkUserWidgetNamedSlot* FUIFrameworkUserWidgetNamedSlotList::AuthorityFindEntry(FName SlotName) const
 {
-	return Slots.FindByPredicate([WidgetId](const FUIFrameworkUserWidgetSlot& Entry) { return Entry.GetWidgetId() == WidgetId; });
+	return Slots.FindByPredicate([SlotName](const FUIFrameworkUserWidgetNamedSlot& Entry) { return Entry.SlotName == SlotName; });
 }
 
-void FUIFrameworkUserWidgetSlotList::ForEachChildren(const TFunctionRef<void(UUIFrameworkWidget*)>& Func)
+void FUIFrameworkUserWidgetNamedSlotList::AuthorityForEachChildren(const TFunctionRef<void(UUIFrameworkWidget*)>& Func)
 {
-	for (FUIFrameworkUserWidgetSlot& Slot : Slots)
+	for (FUIFrameworkUserWidgetNamedSlot& Slot : Slots)
 	{
 		if (UUIFrameworkWidget* ChildWidget = Slot.AuthorityGetWidget())
 		{
@@ -70,11 +84,74 @@ void FUIFrameworkUserWidgetSlotList::ForEachChildren(const TFunctionRef<void(UUI
 	}
 }
 
+/**
+ *
+ */
+void FUIFrameworkUserWidgetViewmodelList::PostReplicatedChange(const TArrayView<int32>& ChangedIndices, int32 FinalSize)
+{
+	check(Owner);
+	if (UUserWidget* UserWidget = Cast<UUserWidget>(Owner->LocalGetUMGWidget()))
+	{
+		if (UMVVMView* View = UMVVMSubsystem::GetViewFromUserWidget(UserWidget))
+		{
+			for (int32 Index : ChangedIndices)
+			{
+				FUIFrameworkUserWidgetViewmodel& Viewmodel = Viewmodels[Index];
+				if (!Viewmodel.Name.IsNone())
+				{
+					View->SetViewModel(Viewmodel.Name, TScriptInterface<INotifyFieldValueChanged>(Viewmodel.Instance));
+				}
+			}
+		}
+	}
+}
+
+bool FUIFrameworkUserWidgetViewmodelList::NetDeltaSerialize(FNetDeltaSerializeInfo& DeltaParms)
+{
+	return FFastArraySerializer::FastArrayDeltaSerialize<FUIFrameworkUserWidgetViewmodel, FUIFrameworkUserWidgetViewmodelList>(Viewmodels, DeltaParms, *this);
+}
+
+void FUIFrameworkUserWidgetViewmodelList::AuthorityAddEntry(FUIFrameworkUserWidgetViewmodel Entry)
+{
+	// Make sure there is only one entry of that name.
+	int32 RemoveCount = Viewmodels.RemoveAllSwap([ViewmodelName = Entry.Name](const FUIFrameworkUserWidgetViewmodel& Entry) { return Entry.Name == ViewmodelName; });
+	if (RemoveCount > 0)
+	{
+		MarkArrayDirty();
+	}
+
+	FUIFrameworkUserWidgetViewmodel& NewEntry = Viewmodels.Add_GetRef(MoveTemp(Entry));
+	MarkItemDirty(NewEntry);
+}
+
+const FUIFrameworkUserWidgetViewmodel* FUIFrameworkUserWidgetViewmodelList::AuthorityFindEntry(FName ViewmodelName) const
+{
+	return Viewmodels.FindByPredicate([ViewmodelName](const FUIFrameworkUserWidgetViewmodel& Entry) { return Entry.Name == ViewmodelName; });
+}
+
+void FUIFrameworkUserWidgetViewmodelList::AttachViewmodels()
+{
+	if (Viewmodels.Num() > 0)
+	{
+		if (UUserWidget* UserWidget = Cast<UUserWidget>(Owner->LocalGetUMGWidget()))
+		{
+			if (UMVVMView* View = UMVVMSubsystem::GetViewFromUserWidget(UserWidget))
+			{
+				for (FUIFrameworkUserWidgetViewmodel& Slot : Viewmodels)
+				{
+					View->SetViewModel(Slot.Name, TScriptInterface<INotifyFieldValueChanged>(Slot.Instance));
+				}
+			}
+		}
+	}
+}
+
 /**
  *
  */
 UUIFrameworkUserWidget::UUIFrameworkUserWidget()
-	: ReplicatedSlotList(this)
+	: ReplicatedNamedSlotList(this)
+	, ReplicatedViewmodelList(this)
 {
 }
 
@@ -84,13 +161,15 @@ void UUIFrameworkUserWidget::GetLifetimeReplicatedProps(TArray<FLifetimeProperty
 
 	FDoRepLifetimeParams Params;
 	Params.bIsPushBased = true;
-	DOREPLIFETIME_WITH_PARAMS_FAST(ThisClass, ReplicatedSlotList, Params);
+	DOREPLIFETIME_WITH_PARAMS_FAST(ThisClass, ReplicatedNamedSlotList, Params);
+	DOREPLIFETIME_WITH_PARAMS_FAST(ThisClass, ReplicatedViewmodelList, Params);
 }
 
 void UUIFrameworkUserWidget::LocalOnUMGWidgetCreated()
 {
 	Super::LocalOnUMGWidgetCreated();
-	//todo add viewmodel
+
+	ReplicatedViewmodelList.AttachViewmodels();
 }
 
 bool UUIFrameworkUserWidget::LocalIsReplicationReady() const
@@ -113,32 +192,71 @@ void UUIFrameworkUserWidget::SetNamedSlot(FName SlotName, UUIFrameworkWidget* Wi
 	else
 	{
 		// Reset the widget to make sure the id is set and it may have been duplicated during the attach
-		FUIFrameworkUserWidgetSlot Entry;
+		FUIFrameworkUserWidgetNamedSlot Entry;
 		Entry.SlotName = SlotName;
 		Entry.AuthoritySetWidget(Widget);
 		Entry.AuthoritySetWidget(FUIFrameworkModule::AuthorityAttachWidget(this, Entry.AuthorityGetWidget()));
-		AddEntry(Entry);
+		ReplicatedNamedSlotList.AuthorityAddEntry(Entry);
 	}
 }
 
+UUIFrameworkWidget* UUIFrameworkUserWidget::GetNamedSlot(FName SlotName) const
+{
+	UUIFrameworkWidget* Result = nullptr;
+	if (SlotName.IsNone())
+	{
+		FFrame::KismetExecutionMessage(TEXT("The slot name is invalid. It can't be added."), ELogVerbosity::Warning, "InvalidWidgetToGet");
+	}
+	else
+	{
+		if (const FUIFrameworkUserWidgetNamedSlot* Slot = ReplicatedNamedSlotList.AuthorityFindEntry(SlotName))
+		{
+			Result = Slot->AuthorityGetWidget();
+		}
+
+	}
+	return Result;
+}
+
+void UUIFrameworkUserWidget::SetViewmodel(FName ViewmodelName, TScriptInterface<INotifyFieldValueChanged> Viewmodel)
+{
+	if (ViewmodelName.IsNone())
+	{
+		FFrame::KismetExecutionMessage(TEXT("The viewmodel name is invalid. It can't be added."), ELogVerbosity::Warning, "InvalidViewmodelToAdd");
+	}
+	else
+	{
+		FUIFrameworkUserWidgetViewmodel NewItem;
+		NewItem.Name = ViewmodelName;
+		NewItem.Instance = Viewmodel.GetObject();
+		ReplicatedViewmodelList.AuthorityAddEntry(NewItem);
+	}
+}
+
+TScriptInterface<INotifyFieldValueChanged> UUIFrameworkUserWidget::GetViewmodel(FName ViewmodelName) const
+{
+	const FUIFrameworkUserWidgetViewmodel* Found = ReplicatedViewmodelList.AuthorityFindEntry(ViewmodelName);
+	return Found ? TScriptInterface<INotifyFieldValueChanged>(Found->Instance) : TScriptInterface<INotifyFieldValueChanged>();
+}
 
 void UUIFrameworkUserWidget::AuthorityForEachChildren(const TFunctionRef<void(UUIFrameworkWidget*)>& Func)
 {
 	Super::AuthorityForEachChildren(Func);
 
-	ReplicatedSlotList.ForEachChildren(Func);
+	ReplicatedNamedSlotList.AuthorityForEachChildren(Func);
 }
 
 void UUIFrameworkUserWidget::AuthorityRemoveChild(UUIFrameworkWidget* Widget)
 {
 	Super::AuthorityRemoveChild(Widget);
-	RemoveEntry(Widget);
+
+	ReplicatedNamedSlotList.AuthorityRemoveEntry(Widget);
 }
 
 void UUIFrameworkUserWidget::LocalAddChild(FUIFrameworkWidgetId ChildId)
 {
 	bool bIsAdded = false;
-	if (FUIFrameworkUserWidgetSlot* NamedSlotEntry = FindEntry(ChildId))
+	if (FUIFrameworkUserWidgetNamedSlot* NamedSlotEntry = ReplicatedNamedSlotList.FindEntry(ChildId))
 	{
 		if (FUIFrameworkWidgetTree* WidgetTree = GetWidgetTree())
 		{
@@ -182,18 +300,3 @@ void UUIFrameworkUserWidget::LocalAddChild(FUIFrameworkWidgetId ChildId)
 		Super::LocalAddChild(ChildId);
 	}
 }
-
-void UUIFrameworkUserWidget::AddEntry(FUIFrameworkUserWidgetSlot Entry)
-{
-	ReplicatedSlotList.AddEntry(Entry);
-}
-
-bool UUIFrameworkUserWidget::RemoveEntry(UUIFrameworkWidget* Widget)
-{
-	return ReplicatedSlotList.RemoveEntry(Widget);
-}
-
-FUIFrameworkUserWidgetSlot* UUIFrameworkUserWidget::FindEntry(FUIFrameworkWidgetId WidgetId)
-{
-	return ReplicatedSlotList.FindEntry(WidgetId);
-}
diff --git a/Engine/Plugins/Experimental/UIFramework/Source/Public/Widgets/UIFUserWidget.h b/Engine/Plugins/Experimental/UIFramework/Source/Public/Widgets/UIFUserWidget.h
--- a/Engine/Plugins/Experimental/UIFramework/Source/Public/Widgets/UIFUserWidget.h
+++ b/Engine/Plugins/Experimental/UIFramework/Source/Public/Widgets/UIFUserWidget.h
@@ -2,6 +2,7 @@
 
 #pragma once
 
+#include "INotifyFieldValueChanged.h"
 #include "UIFWidget.h"
 #include "Types/UIFSlotBase.h"
 
@@ -15,8 +16,8 @@ class UUIFrameworkUserWidget;
 /**
  *
  */
-USTRUCT(BlueprintType)
-struct FUIFrameworkUserWidgetSlot : public FUIFrameworkSlotBase
+USTRUCT()
+struct FUIFrameworkUserWidgetNamedSlot : public FUIFrameworkSlotBase
 {
 	GENERATED_BODY()
 
@@ -29,12 +30,71 @@ struct FUIFrameworkUserWidgetSlot : public FUIFrameworkSlotBase
  *
  */
 USTRUCT()
-struct UIFRAMEWORK_API FUIFrameworkUserWidgetSlotList : public FFastArraySerializer
+struct UIFRAMEWORK_API FUIFrameworkUserWidgetNamedSlotList : public FFastArraySerializer
+{
+	GENERATED_BODY()
+
+	FUIFrameworkUserWidgetNamedSlotList() = default;
+	FUIFrameworkUserWidgetNamedSlotList(UUIFrameworkUserWidget* InOwner)
+		: Owner(InOwner)
+	{}
+
+public:
+	//~ Begin of FFastArraySerializer
+	void PostReplicatedChange(const TArrayView<int32>& ChangedIndices, int32 FinalSize);
+	//~ End of FFastArraySerializer
+
+	bool NetDeltaSerialize(FNetDeltaSerializeInfo& DeltaParms);
+
+	void AuthorityAddEntry(FUIFrameworkUserWidgetNamedSlot Entry);
+	bool AuthorityRemoveEntry(UUIFrameworkWidget* Widget);
+	FUIFrameworkUserWidgetNamedSlot* FindEntry(FUIFrameworkWidgetId WidgetId);
+	const FUIFrameworkUserWidgetNamedSlot* AuthorityFindEntry(FName SlotName) const;
+	void AuthorityForEachChildren(const TFunctionRef<void(UUIFrameworkWidget*)>& Func);
+
+private:
+	UPROPERTY()
+	TArray<FUIFrameworkUserWidgetNamedSlot> Slots;
+
+	UPROPERTY(NotReplicated, Transient)
+	TObjectPtr<UUIFrameworkUserWidget> Owner;
+};
+
+
+template<>
+struct TStructOpsTypeTraits<FUIFrameworkUserWidgetNamedSlotList> : public TStructOpsTypeTraitsBase2<FUIFrameworkUserWidgetNamedSlotList>
+{
+	enum { WithNetDeltaSerializer = true };
+};
+
+/**
+*
+*/
+USTRUCT()
+struct UIFRAMEWORK_API FUIFrameworkUserWidgetViewmodel : public FFastArraySerializerItem
 {
 	GENERATED_BODY()
 
-	FUIFrameworkUserWidgetSlotList() = default;
-	FUIFrameworkUserWidgetSlotList(UUIFrameworkUserWidget* InOwner)
+	FUIFrameworkUserWidgetViewmodel() = default;
+
+public:
+	UPROPERTY()
+	FName Name;
+
+	UPROPERTY()
+	TObjectPtr<UObject> Instance;
+};
+
+/**
+*
+*/
+USTRUCT()
+struct UIFRAMEWORK_API FUIFrameworkUserWidgetViewmodelList : public FFastArraySerializer
+{
+	GENERATED_BODY()
+
+	FUIFrameworkUserWidgetViewmodelList() = default;
+	FUIFrameworkUserWidgetViewmodelList(UUIFrameworkUserWidget* InOwner)
 		: Owner(InOwner)
 	{}
 
@@ -45,14 +105,13 @@ public:
 
 	bool NetDeltaSerialize(FNetDeltaSerializeInfo& DeltaParms);
 
-	void AddEntry(FUIFrameworkUserWidgetSlot Entry);
-	bool RemoveEntry(UUIFrameworkWidget* Widget);
-	FUIFrameworkUserWidgetSlot* FindEntry(FUIFrameworkWidgetId WidgetId);
-	void ForEachChildren(const TFunctionRef<void(UUIFrameworkWidget*)>& Func);
+	void AuthorityAddEntry(FUIFrameworkUserWidgetViewmodel Entry);
+	const FUIFrameworkUserWidgetViewmodel* AuthorityFindEntry(FName ViewmodelName) const;
+	void AttachViewmodels();
 
 private:
 	UPROPERTY()
-	TArray<FUIFrameworkUserWidgetSlot> Slots;
+	TArray<FUIFrameworkUserWidgetViewmodel> Viewmodels;
 
 	UPROPERTY(NotReplicated, Transient)
 	TObjectPtr<UUIFrameworkUserWidget> Owner;
@@ -60,7 +119,7 @@ private:
 
 
 template<>
-struct TStructOpsTypeTraits<FUIFrameworkUserWidgetSlotList> : public TStructOpsTypeTraitsBase2<FUIFrameworkUserWidgetSlotList>
+struct TStructOpsTypeTraits<FUIFrameworkUserWidgetViewmodelList> : public TStructOpsTypeTraitsBase2<FUIFrameworkUserWidgetViewmodelList>
 {
 	enum { WithNetDeltaSerializer = true };
 };
@@ -84,6 +143,15 @@ public:
 	UFUNCTION(BlueprintCallable, BlueprintAuthorityOnly, Category = "UI Framework")
 	void SetNamedSlot(FName SlotName, UUIFrameworkWidget* Widget);
 
+	UFUNCTION(BlueprintCallable, BlueprintAuthorityOnly, Category = "UI Framework")
+	UUIFrameworkWidget* GetNamedSlot(FName SlotName) const;
+
+	UFUNCTION(BlueprintCallable, BlueprintAuthorityOnly, Category = "UI Framework")
+	void SetViewmodel(FName ViewmodelName, TScriptInterface<INotifyFieldValueChanged> Viewmodel);
+
+	UFUNCTION(BlueprintCallable, BlueprintAuthorityOnly, Category = "UI Framework")
+	TScriptInterface<INotifyFieldValueChanged> GetViewmodel(FName ViewmodelName) const;
+
 public:
 	void LocalOnUMGWidgetCreated() override;
 	virtual bool LocalIsReplicationReady() const;
@@ -93,11 +161,9 @@ public:
 	virtual void LocalAddChild(FUIFrameworkWidgetId ChildId) override;
 
 private:
-	void AddEntry(FUIFrameworkUserWidgetSlot Entry);
-	bool RemoveEntry(UUIFrameworkWidget* Widget);
-	FUIFrameworkUserWidgetSlot* FindEntry(FUIFrameworkWidgetId WidgetId);
+	UPROPERTY(Replicated)
+	FUIFrameworkUserWidgetNamedSlotList ReplicatedNamedSlotList;
 
-private:
 	UPROPERTY(Replicated)
-	FUIFrameworkUserWidgetSlotList ReplicatedSlotList;
+	FUIFrameworkUserWidgetViewmodelList ReplicatedViewmodelList;
 };
