UIFramework: Remove the viewmodel code from UIFrameworkUserWidget because it didn't work (missing feature). Add callback to know when the widget changes widget tree (none to new or A to B or A to none).

diff --git a/Engine/Plugins/Experimental/UIFramework/Source/Private/Types/UIFWidgetTree.cpp b/Engine/Plugins/Experimental/UIFramework/Source/Private/Types/UIFWidgetTree.cpp
--- a/Engine/Plugins/Experimental/UIFramework/Source/Private/Types/UIFWidgetTree.cpp
+++ b/Engine/Plugins/Experimental/UIFramework/Source/Private/Types/UIFWidgetTree.cpp
@@ -205,15 +205,15 @@ void FUIFrameworkWidgetTree::AuthorityAddWidget(UUIFrameworkWidget* Parent, UUIF
 
 void FUIFrameworkWidgetTree::AuthorityAddChildInternal(UUIFrameworkWidget* Parent, UUIFrameworkWidget* Child, bool bFirst)
 {
-	Child->WidgetTreeOwner = Owner;
-	check(bFirst || Parent->WidgetTreeOwner == Owner);
+	Child->AuthoritySetWidgetTreeOwner(Owner);
+	check(bFirst || Parent->GetWidgetTreeOwner() == Owner);
 
 	bool bOuterIsDifferent = Child->GetOuter() != ReplicatedOwner;
 	if (bOuterIsDifferent)
 	{
 		//if (Child->GetOuter() == GetTransientPackage())
 		{
-			//If the outer is the transient package, then there are no replication owner yetand it safe to rename it with the correct new outer.
+			//If the outer is the transient package, then there are no replication owner yet and it safe to rename it with the correct new outer.
 			//If the Outer is not the transient, then the widget got replicated with another player.There are no "reset" and we should duplicate and delete the object.
 			//For now only do the rename.It works but it is not the best.
 
@@ -325,7 +325,7 @@ void FUIFrameworkWidgetTree::LocalRemoveRoot(const UUIFrameworkWidget* Widget)
 
 bool FUIFrameworkWidgetTree::AuthorityRemoveChildRecursiveInternal(UUIFrameworkWidget* Widget)
 {
-	Widget->WidgetTreeOwner = nullptr;
+	Widget->AuthoritySetWidgetTreeOwner(nullptr);
 	if (int32* PreviousEntryIndexPtr = AuthorityIndexByWidgetMap.Find(Widget))
 	{
 		check(Entries.IsValidIndex(*PreviousEntryIndexPtr));
diff --git a/Engine/Plugins/Experimental/UIFramework/Source/Private/UIFModule.cpp b/Engine/Plugins/Experimental/UIFramework/Source/Private/UIFModule.cpp
--- a/Engine/Plugins/Experimental/UIFramework/Source/Private/UIFModule.cpp
+++ b/Engine/Plugins/Experimental/UIFramework/Source/Private/UIFModule.cpp
@@ -126,11 +126,11 @@ void FUIFrameworkModule::AuthorityDetachWidgetFromParentInternal(UUIFrameworkWid
 	check(Child);
 
 	// If it's in the WidgetTree, we need to remove it.
-	if (!bTemporary && Child->WidgetTreeOwner)
+	if (!bTemporary && Child->GetWidgetTreeOwner())
 	{
 		//bTemporary: we do not want to remove and re-add the same widget if it's not needed. 
 		//Removing them would cause the local to recreate them instead of re-parenting them.
-		Child->WidgetTreeOwner->GetWidgetTree().AuthorityRemoveWidgetAndChildren(Child);
+		Child->GetWidgetTreeOwner()->GetWidgetTree().AuthorityRemoveWidgetAndChildren(Child);
 	}
 
 	// Notify the widget that we removed a child.
diff --git a/Engine/Plugins/Experimental/UIFramework/Source/Private/Widgets/UIFUserWidget.cpp b/Engine/Plugins/Experimental/UIFramework/Source/Private/Widgets/UIFUserWidget.cpp
--- a/Engine/Plugins/Experimental/UIFramework/Source/Private/Widgets/UIFUserWidget.cpp
+++ b/Engine/Plugins/Experimental/UIFramework/Source/Private/Widgets/UIFUserWidget.cpp
@@ -84,74 +84,11 @@ void FUIFrameworkUserWidgetNamedSlotList::AuthorityForEachChildren(const TFuncti
 	}
 }
 
-/**
- *
- */
-void FUIFrameworkUserWidgetViewmodelList::PostReplicatedChange(const TArrayView<int32>& ChangedIndices, int32 FinalSize)
-{
-	check(Owner);
-	if (UUserWidget* UserWidget = Cast<UUserWidget>(Owner->LocalGetUMGWidget()))
-	{
-		if (UMVVMView* View = UMVVMSubsystem::GetViewFromUserWidget(UserWidget))
-		{
-			for (int32 Index : ChangedIndices)
-			{
-				FUIFrameworkUserWidgetViewmodel& Viewmodel = Viewmodels[Index];
-				if (!Viewmodel.Name.IsNone())
-				{
-					View->SetViewModel(Viewmodel.Name, TScriptInterface<INotifyFieldValueChanged>(Viewmodel.Instance));
-				}
-			}
-		}
-	}
-}
-
-bool FUIFrameworkUserWidgetViewmodelList::NetDeltaSerialize(FNetDeltaSerializeInfo& DeltaParms)
-{
-	return FFastArraySerializer::FastArrayDeltaSerialize<FUIFrameworkUserWidgetViewmodel, FUIFrameworkUserWidgetViewmodelList>(Viewmodels, DeltaParms, *this);
-}
-
-void FUIFrameworkUserWidgetViewmodelList::AuthorityAddEntry(FUIFrameworkUserWidgetViewmodel Entry)
-{
-	// Make sure there is only one entry of that name.
-	int32 RemoveCount = Viewmodels.RemoveAllSwap([ViewmodelName = Entry.Name](const FUIFrameworkUserWidgetViewmodel& Entry) { return Entry.Name == ViewmodelName; });
-	if (RemoveCount > 0)
-	{
-		MarkArrayDirty();
-	}
-
-	FUIFrameworkUserWidgetViewmodel& NewEntry = Viewmodels.Add_GetRef(MoveTemp(Entry));
-	MarkItemDirty(NewEntry);
-}
-
-const FUIFrameworkUserWidgetViewmodel* FUIFrameworkUserWidgetViewmodelList::AuthorityFindEntry(FName ViewmodelName) const
-{
-	return Viewmodels.FindByPredicate([ViewmodelName](const FUIFrameworkUserWidgetViewmodel& Entry) { return Entry.Name == ViewmodelName; });
-}
-
-void FUIFrameworkUserWidgetViewmodelList::AttachViewmodels()
-{
-	if (Viewmodels.Num() > 0)
-	{
-		if (UUserWidget* UserWidget = Cast<UUserWidget>(Owner->LocalGetUMGWidget()))
-		{
-			if (UMVVMView* View = UMVVMSubsystem::GetViewFromUserWidget(UserWidget))
-			{
-				for (FUIFrameworkUserWidgetViewmodel& Slot : Viewmodels)
-				{
-					View->SetViewModel(Slot.Name, TScriptInterface<INotifyFieldValueChanged>(Slot.Instance));
-				}
-			}
-		}
-	}
-}
-
 /**
  *
  */
 UUIFrameworkUserWidget::UUIFrameworkUserWidget()
 	: ReplicatedNamedSlotList(this)
-	, ReplicatedViewmodelList(this)
 {
 }
 
@@ -162,14 +99,6 @@ void UUIFrameworkUserWidget::GetLifetimeReplicatedProps(TArray<FLifetimeProperty
 	FDoRepLifetimeParams Params;
 	Params.bIsPushBased = true;
 	DOREPLIFETIME_WITH_PARAMS_FAST(ThisClass, ReplicatedNamedSlotList, Params);
-	DOREPLIFETIME_WITH_PARAMS_FAST(ThisClass, ReplicatedViewmodelList, Params);
-}
-
-void UUIFrameworkUserWidget::LocalOnUMGWidgetCreated()
-{
-	Super::LocalOnUMGWidgetCreated();
-
-	ReplicatedViewmodelList.AttachViewmodels();
 }
 
 bool UUIFrameworkUserWidget::LocalIsReplicationReady() const
@@ -218,27 +147,6 @@ UUIFrameworkWidget* UUIFrameworkUserWidget::GetNamedSlot(FName SlotName) const
 	return Result;
 }
 
-void UUIFrameworkUserWidget::SetViewmodel(FName ViewmodelName, TScriptInterface<INotifyFieldValueChanged> Viewmodel)
-{
-	if (ViewmodelName.IsNone())
-	{
-		FFrame::KismetExecutionMessage(TEXT("The viewmodel name is invalid. It can't be added."), ELogVerbosity::Warning, "InvalidViewmodelToAdd");
-	}
-	else
-	{
-		FUIFrameworkUserWidgetViewmodel NewItem;
-		NewItem.Name = ViewmodelName;
-		NewItem.Instance = Viewmodel.GetObject();
-		ReplicatedViewmodelList.AuthorityAddEntry(NewItem);
-	}
-}
-
-TScriptInterface<INotifyFieldValueChanged> UUIFrameworkUserWidget::GetViewmodel(FName ViewmodelName) const
-{
-	const FUIFrameworkUserWidgetViewmodel* Found = ReplicatedViewmodelList.AuthorityFindEntry(ViewmodelName);
-	return Found ? TScriptInterface<INotifyFieldValueChanged>(Found->Instance) : TScriptInterface<INotifyFieldValueChanged>();
-}
-
 void UUIFrameworkUserWidget::AuthorityForEachChildren(const TFunctionRef<void(UUIFrameworkWidget*)>& Func)
 {
 	Super::AuthorityForEachChildren(Func);
diff --git a/Engine/Plugins/Experimental/UIFramework/Source/Public/Types/UIFWidgetTree.h b/Engine/Plugins/Experimental/UIFramework/Source/Public/Types/UIFWidgetTree.h
--- a/Engine/Plugins/Experimental/UIFramework/Source/Public/Types/UIFWidgetTree.h
+++ b/Engine/Plugins/Experimental/UIFramework/Source/Public/Types/UIFWidgetTree.h
@@ -82,6 +82,11 @@ public:
 		return FFastArraySerializer::FastArrayDeltaSerialize<FUIFrameworkWidgetTreeEntry, FUIFrameworkWidgetTree>(Entries, DeltaParms, *this);
 	}
 
+	AActor* GetReplicationOwner() const
+	{
+		return ReplicatedOwner;
+	}
+
 	bool ReplicateSubWidgets(UActorChannel* Channel, FOutBunch* Bunch, FReplicationFlags* RepFlags);
 
 	/** Add a new widget to the top hierarchy. */
diff --git a/Engine/Plugins/Experimental/UIFramework/Source/Public/UIFWidget.h b/Engine/Plugins/Experimental/UIFramework/Source/Public/UIFWidget.h
--- a/Engine/Plugins/Experimental/UIFramework/Source/Public/UIFWidget.h
+++ b/Engine/Plugins/Experimental/UIFramework/Source/Public/UIFWidget.h
@@ -28,9 +28,6 @@ class UUIFrameworkWidgetWrapperInterface : public UInterface
 class IUIFrameworkWidgetWrapperInterface
 {
 	GENERATED_BODY()
-
-public:
-	virtual void ReplaceWidget(UUIFrameworkWidget* OldWidget, UUIFrameworkWidget* NewWidget) {}
 };
 
 /**
@@ -63,12 +60,12 @@ public:
 
 	TScriptInterface<IUIFrameworkWidgetWrapperInterface> AuthorityGetWrapper() const
 	{
-		return Wrapper;
+		return AuthorityWrapper;
 	}
 
 	void AuthoritySetWrapper(TScriptInterface<IUIFrameworkWidgetWrapperInterface> InWrapper)
 	{
-		Wrapper = InWrapper;
+		AuthorityWrapper = InWrapper;
 	}
 
 	FUIFrameworkWidgetId GetWidgetId() const
@@ -124,6 +121,9 @@ protected:
 	virtual void AuthorityRemoveChild(UUIFrameworkWidget* Widget)
 	{
 	}
+	virtual void AuthorityOnWidgetTreeOwnerChanged()
+	{
+	}
 	virtual void LocalOnUMGWidgetCreated()
 	{
 	}
@@ -136,6 +136,15 @@ private:
 	UFUNCTION()
 	void OnRep_Visibility();
 
+	void AuthoritySetWidgetTreeOwner(IUIFrameworkWidgetTreeOwner* InWidgetTreeOwner)
+	{
+		if (WidgetTreeOwner != InWidgetTreeOwner)
+		{
+			WidgetTreeOwner = InWidgetTreeOwner;
+			AuthorityOnWidgetTreeOwnerChanged();
+		}
+	}
+
 protected:
 	UPROPERTY(BlueprintReadOnly, Replicated, EditDefaultsOnly, Category = "UI Framework")
 	TSoftClassPtr<UWidget> WidgetClass; // todo: make this private and use a constructor argument
@@ -147,7 +156,7 @@ private:
 
 	//~ Authority
 	UPROPERTY(Transient)
-	TScriptInterface<IUIFrameworkWidgetWrapperInterface> Wrapper;
+	TScriptInterface<IUIFrameworkWidgetWrapperInterface> AuthorityWrapper;
 
 	//~ Authority and Local
 	IUIFrameworkWidgetTreeOwner* WidgetTreeOwner = nullptr;
diff --git a/Engine/Plugins/Experimental/UIFramework/Source/Public/Widgets/UIFUserWidget.h b/Engine/Plugins/Experimental/UIFramework/Source/Public/Widgets/UIFUserWidget.h
--- a/Engine/Plugins/Experimental/UIFramework/Source/Public/Widgets/UIFUserWidget.h
+++ b/Engine/Plugins/Experimental/UIFramework/Source/Public/Widgets/UIFUserWidget.h
@@ -67,63 +67,6 @@ struct TStructOpsTypeTraits<FUIFrameworkUserWidgetNamedSlotList> : public TStruc
 	enum { WithNetDeltaSerializer = true };
 };
 
-/**
-*
-*/
-USTRUCT()
-struct UIFRAMEWORK_API FUIFrameworkUserWidgetViewmodel : public FFastArraySerializerItem
-{
-	GENERATED_BODY()
-
-	FUIFrameworkUserWidgetViewmodel() = default;
-
-public:
-	UPROPERTY()
-	FName Name;
-
-	UPROPERTY()
-	TObjectPtr<UObject> Instance;
-};
-
-/**
-*
-*/
-USTRUCT()
-struct UIFRAMEWORK_API FUIFrameworkUserWidgetViewmodelList : public FFastArraySerializer
-{
-	GENERATED_BODY()
-
-	FUIFrameworkUserWidgetViewmodelList() = default;
-	FUIFrameworkUserWidgetViewmodelList(UUIFrameworkUserWidget* InOwner)
-		: Owner(InOwner)
-	{}
-
-public:
-	//~ Begin of FFastArraySerializer
-	void PostReplicatedChange(const TArrayView<int32>& ChangedIndices, int32 FinalSize);
-	//~ End of FFastArraySerializer
-
-	bool NetDeltaSerialize(FNetDeltaSerializeInfo& DeltaParms);
-
-	void AuthorityAddEntry(FUIFrameworkUserWidgetViewmodel Entry);
-	const FUIFrameworkUserWidgetViewmodel* AuthorityFindEntry(FName ViewmodelName) const;
-	void AttachViewmodels();
-
-private:
-	UPROPERTY()
-	TArray<FUIFrameworkUserWidgetViewmodel> Viewmodels;
-
-	UPROPERTY(NotReplicated, Transient)
-	TObjectPtr<UUIFrameworkUserWidget> Owner;
-};
-
-
-template<>
-struct TStructOpsTypeTraits<FUIFrameworkUserWidgetViewmodelList> : public TStructOpsTypeTraitsBase2<FUIFrameworkUserWidgetViewmodelList>
-{
-	enum { WithNetDeltaSerializer = true };
-};
-
 /**
  *
  */
@@ -136,7 +79,6 @@ public:
 	UUIFrameworkUserWidget();
 
 public:
-
 	UFUNCTION(BlueprintCallable, BlueprintAuthorityOnly, Category = "UI Framework")
 	void SetWidgetClass(TSoftClassPtr<UWidget> Value);
 
@@ -146,15 +88,8 @@ public:
 	UFUNCTION(BlueprintCallable, BlueprintAuthorityOnly, Category = "UI Framework")
 	UUIFrameworkWidget* GetNamedSlot(FName SlotName) const;
 
-	UFUNCTION(BlueprintCallable, BlueprintAuthorityOnly, Category = "UI Framework")
-	void SetViewmodel(FName ViewmodelName, TScriptInterface<INotifyFieldValueChanged> Viewmodel);
-
-	UFUNCTION(BlueprintCallable, BlueprintAuthorityOnly, Category = "UI Framework")
-	TScriptInterface<INotifyFieldValueChanged> GetViewmodel(FName ViewmodelName) const;
-
 public:
-	void LocalOnUMGWidgetCreated() override;
-	virtual bool LocalIsReplicationReady() const;
+	virtual bool LocalIsReplicationReady() const override;
 
 	virtual void AuthorityForEachChildren(const TFunctionRef<void(UUIFrameworkWidget*)>& Func) override;
 	virtual void AuthorityRemoveChild(UUIFrameworkWidget* Widget) override;
@@ -163,7 +98,4 @@ public:
 private:
 	UPROPERTY(Replicated)
 	FUIFrameworkUserWidgetNamedSlotList ReplicatedNamedSlotList;
-
-	UPROPERTY(Replicated)
-	FUIFrameworkUserWidgetViewmodelList ReplicatedViewmodelList;
 };
