UIFramework: Map of object. When the object is GC, the map cannot find the object since it is not valid anymore. Use FObjectKey instead.

diff --git a/Engine/Plugins/Experimental/UIFramework/Source/Private/Types/UIFWidgetTree.cpp b/Engine/Plugins/Experimental/UIFramework/Source/Private/Types/UIFWidgetTree.cpp
--- a/Engine/Plugins/Experimental/UIFramework/Source/Private/Types/UIFWidgetTree.cpp
+++ b/Engine/Plugins/Experimental/UIFramework/Source/Private/Types/UIFWidgetTree.cpp
@@ -205,6 +205,11 @@ void FUIFrameworkWidgetTree::AuthorityAddWidget(UUIFrameworkWidget* Parent, UUIF
 
 void FUIFrameworkWidgetTree::AuthorityAddChildInternal(UUIFrameworkWidget* Parent, UUIFrameworkWidget* Child, bool bFirst)
 {
+	if (Child == nullptr)
+	{
+		return;
+	}
+
 	Child->AuthoritySetWidgetTreeOwner(Owner);
 	check(bFirst || Parent->GetWidgetTreeOwner() == Owner);
 
@@ -231,7 +236,8 @@ void FUIFrameworkWidgetTree::AuthorityAddChildInternal(UUIFrameworkWidget* Paren
 		}
 	}
 
-	if (int32* PreviousEntryIndexPtr = AuthorityIndexByWidgetMap.Find(Child))
+	const FObjectKey ChildKey = Child;
+	if (int32* PreviousEntryIndexPtr = AuthorityIndexByWidgetMap.Find(ChildKey))
 	{
 		check(Entries.IsValidIndex(*PreviousEntryIndexPtr));
 		FUIFrameworkWidgetTreeEntry& PreviousEntry = Entries[*PreviousEntryIndexPtr];
@@ -254,7 +260,7 @@ void FUIFrameworkWidgetTree::AuthorityAddChildInternal(UUIFrameworkWidget* Paren
 		FUIFrameworkWidgetTreeEntry& NewEntry = Entries[NewEntryIndex];
 		MarkItemDirty(NewEntry);
 
-		AuthorityIndexByWidgetMap.Add(Child) = NewEntryIndex;
+		AuthorityIndexByWidgetMap.Add(ChildKey) = NewEntryIndex;
 		WidgetByIdMap.Add(Child->GetWidgetId()) = Child;
 
 		if (ensure(ReplicatedOwner) && ReplicatedOwner->IsUsingRegisteredSubObjectList())
@@ -325,12 +331,18 @@ void FUIFrameworkWidgetTree::LocalRemoveRoot(const UUIFrameworkWidget* Widget)
 
 bool FUIFrameworkWidgetTree::AuthorityRemoveChildRecursiveInternal(UUIFrameworkWidget* Widget)
 {
+	if (Widget == nullptr)
+	{
+		return false;
+	}
+
 	Widget->AuthoritySetWidgetTreeOwner(nullptr);
-	if (int32* PreviousEntryIndexPtr = AuthorityIndexByWidgetMap.Find(Widget))
+	const FObjectKey WidgetKey = Widget;
+	if (int32* PreviousEntryIndexPtr = AuthorityIndexByWidgetMap.Find(WidgetKey))
 	{
 		check(Entries.IsValidIndex(*PreviousEntryIndexPtr));
 
-		AuthorityIndexByWidgetMap.Remove(Widget);
+		AuthorityIndexByWidgetMap.Remove(WidgetKey);
 		WidgetByIdMap.Remove(Widget->GetWidgetId());
 
 		if (ensure(ReplicatedOwner) && ReplicatedOwner->IsUsingRegisteredSubObjectList())
@@ -343,7 +355,7 @@ bool FUIFrameworkWidgetTree::AuthorityRemoveChildRecursiveInternal(UUIFrameworkW
 		// Fix up the swap item
 		if (Entries.IsValidIndex(*PreviousEntryIndexPtr))
 		{
-			if (int32* FixUpChildIndexPtr = AuthorityIndexByWidgetMap.Find(Entries[*PreviousEntryIndexPtr].Child))
+			if (int32* FixUpChildIndexPtr = AuthorityIndexByWidgetMap.Find(FObjectKey(Entries[*PreviousEntryIndexPtr].Child)))
 			{
 				*FixUpChildIndexPtr = *PreviousEntryIndexPtr;
 			}
@@ -466,7 +478,7 @@ void FUIFrameworkWidgetTree::AuthorityTest() const
 		}
 
 
-		const int32* FoundIndexPtr = AuthorityIndexByWidgetMap.Find(Entry.Child);
+		const int32* FoundIndexPtr = AuthorityIndexByWidgetMap.Find(FObjectKey(Entry.Child));
 		if (FoundIndexPtr)
 		{
 			int32 FoundIndex = *FoundIndexPtr;
diff --git a/Engine/Plugins/Experimental/UIFramework/Source/Public/Types/UIFWidgetTree.h b/Engine/Plugins/Experimental/UIFramework/Source/Public/Types/UIFWidgetTree.h
--- a/Engine/Plugins/Experimental/UIFramework/Source/Public/Types/UIFWidgetTree.h
+++ b/Engine/Plugins/Experimental/UIFramework/Source/Public/Types/UIFWidgetTree.h
@@ -149,7 +149,7 @@ private:
 	UPROPERTY(NotReplicated, Transient)
 	TObjectPtr<AActor> ReplicatedOwner;
 
-	TMap<TWeakObjectPtr<UUIFrameworkWidget>, int32> AuthorityIndexByWidgetMap;
+	TMap<FObjectKey, int32> AuthorityIndexByWidgetMap;
 	TMap<FUIFrameworkWidgetId, TWeakObjectPtr<UUIFrameworkWidget>> WidgetByIdMap;
 	IUIFrameworkWidgetTreeOwner* Owner = nullptr;
 };
