Fix missing include introduced with CL 34029570
#rnx

diff --git a/Engine/Plugins/Experimental/UIFramework/Source/Public/Types/UIFWidgetTree.h b/Engine/Plugins/Experimental/UIFramework/Source/Public/Types/UIFWidgetTree.h
--- a/Engine/Plugins/Experimental/UIFramework/Source/Public/Types/UIFWidgetTree.h
+++ b/Engine/Plugins/Experimental/UIFramework/Source/Public/Types/UIFWidgetTree.h
@@ -4,6 +4,7 @@
 
 #include "Net/Serialization/FastArraySerializer.h"
 #include "Types/UIFWidgetId.h"
+#include "UObject/ObjectKey.h"
 
 #include "UIFWidgetTree.generated.h"
 
