UIFramework: Allow UIFPlayerComponents to be attached to actors owned by a player controller

- Also fix a crash that would occur when there wasn't a valid widget owner

diff --git a/Engine/Plugins/Experimental/UIFramework/Source/Private/UIFPlayerComponent.cpp b/Engine/Plugins/Experimental/UIFramework/Source/Private/UIFPlayerComponent.cpp
--- a/Engine/Plugins/Experimental/UIFramework/Source/Private/UIFPlayerComponent.cpp
+++ b/Engine/Plugins/Experimental/UIFramework/Source/Private/UIFPlayerComponent.cpp
@@ -157,7 +157,7 @@ bool UUIFrameworkPlayerComponent::ReplicateSubobjects(UActorChannel* Channel, cl
 void UUIFrameworkPlayerComponent::AddWidget(FUIFrameworkGameLayerSlot InEntry)
 {
 	APlayerController* LocalOwner = GetPlayerController();
-	check(LocalOwner->HasAuthority());
+	check(!LocalOwner || LocalOwner->HasAuthority());
 
 	if (InEntry.AuthorityGetWidget() == nullptr)
 	{
@@ -174,7 +174,7 @@ void UUIFrameworkPlayerComponent::AddWidget(FUIFrameworkGameLayerSlot InEntry)
 void UUIFrameworkPlayerComponent::RemoveWidget(UUIFrameworkWidget* Widget)
 {
 	APlayerController* LocalOwner = GetPlayerController();
-	check(LocalOwner->HasAuthority());
+	check(!LocalOwner || LocalOwner->HasAuthority());
 
 	if (Widget == nullptr)
 	{
@@ -222,7 +222,7 @@ FUIFrameworkWidgetOwner UUIFrameworkPlayerComponent::GetWidgetOwner() const
 void UUIFrameworkPlayerComponent::TickComponent(float DeltaTime, enum ELevelTick TickType, FActorComponentTickFunction* ThisTickFunction)
 {
 	APlayerController* LocalOwner = GetPlayerController();
-	check(!LocalOwner->HasAuthority());
+	check(!LocalOwner || !LocalOwner->HasAuthority());
 
 	if (ClassesToLoad.Num() == 0 && NetReplicationPending.Num() == 0)
 	{
@@ -282,7 +282,7 @@ void UUIFrameworkPlayerComponent::TickComponent(float DeltaTime, enum ELevelTick
 void UUIFrameworkPlayerComponent::LocalWidgetWasAddedToTree(const FUIFrameworkWidgetTreeEntry& Entry)
 {
 	APlayerController* LocalOwner = GetPlayerController();
-	check(!LocalOwner->HasAuthority());
+	check(!LocalOwner || !LocalOwner->HasAuthority());
 
 	if (Entry.Child && Entry.Child->LocalIsReplicationReady())
 	{
@@ -343,7 +343,7 @@ void UUIFrameworkPlayerComponent::LocalWidgetWasAddedToTree(const FUIFrameworkWi
 void UUIFrameworkPlayerComponent::LocalWidgetRemovedFromTree(const FUIFrameworkWidgetTreeEntry& Entry)
 {
 	APlayerController* LocalOwner = GetPlayerController();
-	check(!LocalOwner->HasAuthority());
+	check(!LocalOwner || !LocalOwner->HasAuthority());
 	check(!bAddingWidget);
 
 	NetReplicationPending.Remove(Entry.ReplicationID);
diff --git a/Engine/Plugins/Experimental/UIFramework/Source/Private/UIFWidget.cpp b/Engine/Plugins/Experimental/UIFramework/Source/Private/UIFWidget.cpp
--- a/Engine/Plugins/Experimental/UIFramework/Source/Private/UIFWidget.cpp
+++ b/Engine/Plugins/Experimental/UIFramework/Source/Private/UIFWidget.cpp
@@ -128,9 +128,13 @@ void UUIFrameworkWidget::LocalCreateUMGWidget(TNonNullPtr<IUIFrameworkWidgetTree
 			check(Class->IsChildOf(UWidget::StaticClass()));
 			LocalUMGWidget = NewObject<UWidget>(this, Class, FName(), RF_Transient);
 		}
-		LocalUMGWidget->SetIsEnabled(bIsEnabled);
-		LocalUMGWidget->SetVisibility(Visibility);
-		LocalOnUMGWidgetCreated();
+
+		if (LocalUMGWidget != nullptr)
+		{
+			LocalUMGWidget->SetIsEnabled(bIsEnabled);
+			LocalUMGWidget->SetVisibility(Visibility);
+			LocalOnUMGWidgetCreated();
+		}
 	}
 }
 
diff --git a/Engine/Plugins/Experimental/UIFramework/Source/Public/UIFPlayerComponent.h b/Engine/Plugins/Experimental/UIFramework/Source/Public/UIFPlayerComponent.h
--- a/Engine/Plugins/Experimental/UIFramework/Source/Public/UIFPlayerComponent.h
+++ b/Engine/Plugins/Experimental/UIFramework/Source/Public/UIFPlayerComponent.h
@@ -136,14 +136,26 @@ public:
 	T* GetPlayerController() const
 	{
 		static_assert(TPointerIsConvertibleFromTo<T, APlayerController>::Value, "'T' template parameter to GetPlayerController must be derived from APlayerController");
-		return Cast<T>(GetOwner());
+		AActor* Owner = GetOwner();
+		while (Owner != nullptr)
+		{
+			if (T* PC = Cast<T>(Owner))
+			{
+				return PC;
+			}
+
+			Owner = Owner->GetOwner();
+		}
+		return nullptr;
 	}
 
 	template <class T = APlayerController>
 	T* GetPlayerControllerChecked() const
 	{
 		static_assert(TPointerIsConvertibleFromTo<T, APlayerController>::Value, "'T' template parameter to GetPlayerControllerChecked must be derived from APlayerController");
-		return CastChecked<T>(GetOwner());
+		T* Result = GetPlayerController<T>();
+		check(Result);
+		return Result;
 	}
 
 	//~ Begin UActorComponent
