UIFramework: add a callback for when a widget is added to the tree locally

diff --git a/Engine/Plugins/Experimental/UIFramework/Source/Private/Types/UIFWidgetTree.cpp b/Engine/Plugins/Experimental/UIFramework/Source/Private/Types/UIFWidgetTree.cpp
--- a/Engine/Plugins/Experimental/UIFramework/Source/Private/Types/UIFWidgetTree.cpp
+++ b/Engine/Plugins/Experimental/UIFramework/Source/Private/Types/UIFWidgetTree.cpp
@@ -142,6 +142,7 @@ void FUIFrameworkWidgetTree::PostReplicatedAdd(const TArrayView<int32> AddedIndi
 			{
 				WidgetByIdMap.FindOrAdd(Entry.ParentId) = Entry.Parent;
 			}
+ 			LocalOnWidgetAdded.Broadcast(Entry.Child);
 		}
 	}
 }
diff --git a/Engine/Plugins/Experimental/UIFramework/Source/Public/Types/UIFWidgetTree.h b/Engine/Plugins/Experimental/UIFramework/Source/Public/Types/UIFWidgetTree.h
--- a/Engine/Plugins/Experimental/UIFramework/Source/Public/Types/UIFWidgetTree.h
+++ b/Engine/Plugins/Experimental/UIFramework/Source/Public/Types/UIFWidgetTree.h
@@ -131,6 +131,7 @@ public:
 	DECLARE_MULTICAST_DELEGATE_OneParam(FUIFrameworkWidgetDelegate, UUIFrameworkWidget*);
 	FUIFrameworkWidgetDelegate AuthorityOnWidgetAdded;
 	FUIFrameworkWidgetDelegate AuthorityOnWidgetRemoved;
+	FUIFrameworkWidgetDelegate LocalOnWidgetAdded;
 
 private:
 	void AuthorityAddChildInternal(UUIFrameworkWidget* Parent, UUIFrameworkWidget* Child, bool bFirst);
