Verse UI: add a property to allow changing the hit-test visibility of widgets.
#okforversepublic

diff --git a/Engine/Plugins/Experimental/UIFramework/Source/Private/UIFWidget.cpp b/Engine/Plugins/Experimental/UIFramework/Source/Private/UIFWidget.cpp
--- a/Engine/Plugins/Experimental/UIFramework/Source/Private/UIFWidget.cpp
+++ b/Engine/Plugins/Experimental/UIFramework/Source/Private/UIFWidget.cpp
@@ -157,6 +157,11 @@ ESlateVisibility UUIFrameworkWidget::GetVisibility() const
 
 void UUIFrameworkWidget::SetVisibility(ESlateVisibility InVisibility)
 {
+	if (InVisibility == ESlateVisibility::Visible && !bIsHitTestVisible)
+	{
+		InVisibility = ESlateVisibility::SelfHitTestInvisible;
+	}
+
 	if (Visibility != InVisibility)
 	{
 		Visibility = InVisibility;
@@ -180,6 +185,28 @@ void UUIFrameworkWidget::SetEnabled(bool bInIsEnabled)
 	}
 }
 
+bool UUIFrameworkWidget::IsHitTestVisible() const
+{
+	return bIsHitTestVisible;
+}
+
+void UUIFrameworkWidget::SetHitTestVisible(bool bInHitTestVisible)
+{
+	if (bIsHitTestVisible != bInHitTestVisible)
+	{
+		bIsHitTestVisible = bInHitTestVisible;
+
+		if (Visibility == ESlateVisibility::Visible && !bIsHitTestVisible)
+		{
+			SetVisibility(ESlateVisibility::SelfHitTestInvisible);
+		}
+		else if (Visibility == ESlateVisibility::SelfHitTestInvisible && bIsHitTestVisible)
+		{
+			SetVisibility(ESlateVisibility::Visible);
+		}
+	}
+}
+
 void UUIFrameworkWidget::OnRep_IsEnabled()
 {
 	if (LocalUMGWidget)
diff --git a/Engine/Plugins/Experimental/UIFramework/Source/Public/UIFWidget.h b/Engine/Plugins/Experimental/UIFramework/Source/Public/UIFWidget.h
--- a/Engine/Plugins/Experimental/UIFramework/Source/Public/UIFWidget.h
+++ b/Engine/Plugins/Experimental/UIFramework/Source/Public/UIFWidget.h
@@ -48,6 +48,9 @@ private:
 	UPROPERTY(BlueprintReadWrite, ReplicatedUsing = "OnRep_Visibility", Getter, Setter = "SetVisibility", Category = "UI Framework", meta = (AllowPrivateAccess = "true"))
 	ESlateVisibility Visibility = ESlateVisibility::Visible;
 
+	UPROPERTY(BlueprintReadWrite, Getter = "IsHitTestVisible", Setter = "SetHitTestVisible", Category = "UI Framework", meta = (AllowPrivateAccess = "true"))
+	bool bIsHitTestVisible = true;
+
 public:
 	//~ Begin UObject
 	virtual bool IsSupportedForNetworking() const override
@@ -117,6 +120,9 @@ public:
 	bool IsEnabled() const;
 	void SetEnabled(bool bEnabled);
 
+	bool IsHitTestVisible() const;
+	void SetHitTestVisible(bool bInHitTestVisible);
+
 protected:
 	virtual void AuthorityRemoveChild(UUIFrameworkWidget* Widget)
 	{
