UIFramework: Fix Button Clicks not working when the UIFPlayerComponent is not on the player controller

- UUIFrameworkWidget instances were making assumptions about their outer
  - Added `UUIFrameworkWidget::GetPlayerController()` to more reliably grab the widget owner's player controller

diff --git a/Engine/Plugins/Experimental/UIFramework/Source/Private/Widgets/UIFButton.cpp b/Engine/Plugins/Experimental/UIFramework/Source/Private/Widgets/UIFButton.cpp
--- a/Engine/Plugins/Experimental/UIFramework/Source/Private/Widgets/UIFButton.cpp
+++ b/Engine/Plugins/Experimental/UIFramework/Source/Private/Widgets/UIFButton.cpp
@@ -133,7 +133,7 @@ void UUIFrameworkButton::LocalAddChild(FUIFrameworkWidgetId ChildId)
 void UUIFrameworkButton::HandleClick()
 {
 	// todo the click event should send the userid
-	ServerClick(Cast<APlayerController>(GetOuter()));
+	ServerClick(GetPlayerController());
 }
 
 void UUIFrameworkButton::ServerClick_Implementation(APlayerController* PlayerController)
diff --git a/Engine/Plugins/Experimental/UIFramework/Source/Public/UIFWidget.h b/Engine/Plugins/Experimental/UIFramework/Source/Public/UIFWidget.h
--- a/Engine/Plugins/Experimental/UIFramework/Source/Public/UIFWidget.h
+++ b/Engine/Plugins/Experimental/UIFramework/Source/Public/UIFWidget.h
@@ -61,6 +61,32 @@ public:
 	virtual bool CallRemoteFunction(UFunction* Function, void* Parameters, FOutParmRec* OutParms, FFrame* Stack) override;
 	//~ End UObject
 
+	// Gets the controller that owns the widget
+	template <class T = APlayerController>
+	T* GetPlayerController() const
+	{
+		static_assert(TPointerIsConvertibleFromTo<T, APlayerController>::Value, "'T' template parameter to GetPlayerController must be derived from APlayerController");
+		UObject* TestOuter = GetOuter();
+		while (TestOuter != nullptr)
+		{
+			if (T* PC = Cast<T>(TestOuter))
+			{
+				return PC;
+			}
+
+			const AActor* OuterActor = Cast<AActor>(TestOuter);
+			if (OuterActor && OuterActor->GetOwner())
+			{
+				TestOuter = OuterActor->GetOwner();
+			}
+			else
+			{
+				TestOuter = TestOuter->GetOuter();
+			}
+		}
+		return nullptr;
+	}
+
 	TScriptInterface<IUIFrameworkWidgetWrapperInterface> AuthorityGetWrapper() const
 	{
 		return AuthorityWrapper;
