Add RenderOpacity to Widget class in Verse (epic_internal).

#rnx

#okforversepublic

diff --git a/Engine/Plugins/Experimental/UIFramework/Source/Private/UIFWidget.cpp b/Engine/Plugins/Experimental/UIFramework/Source/Private/UIFWidget.cpp
--- a/Engine/Plugins/Experimental/UIFramework/Source/Private/UIFWidget.cpp
+++ b/Engine/Plugins/Experimental/UIFramework/Source/Private/UIFWidget.cpp
@@ -207,6 +207,19 @@ void UUIFrameworkWidget::SetHitTestVisible(bool bInHitTestVisible)
 	}
 }
 
+double UUIFrameworkWidget::GetRenderOpacity() const
+{
+	return RenderOpacity;
+}
+
+void UUIFrameworkWidget::SetRenderOpacity(double InRenderOpacity)
+{
+	if (RenderOpacity != InRenderOpacity)
+	{
+		RenderOpacity = InRenderOpacity;
+	}
+}
+
 void UUIFrameworkWidget::OnRep_IsEnabled()
 {
 	if (LocalUMGWidget)
@@ -222,3 +235,11 @@ void UUIFrameworkWidget::OnRep_Visibility()
 		LocalUMGWidget->SetVisibility(Visibility);
 	}
 }
+
+void UUIFrameworkWidget::OnRep_RenderOpacity()
+{
+	if (LocalUMGWidget)
+	{
+		LocalUMGWidget->SetRenderOpacity((float)RenderOpacity);
+	}
+}
diff --git a/Engine/Plugins/Experimental/UIFramework/Source/Public/UIFWidget.h b/Engine/Plugins/Experimental/UIFramework/Source/Public/UIFWidget.h
--- a/Engine/Plugins/Experimental/UIFramework/Source/Public/UIFWidget.h
+++ b/Engine/Plugins/Experimental/UIFramework/Source/Public/UIFWidget.h
@@ -51,6 +51,9 @@ private:
 	UPROPERTY(BlueprintReadWrite, Getter = "IsHitTestVisible", Setter = "SetHitTestVisible", Category = "UI Framework", meta = (AllowPrivateAccess = "true"))
 	bool bIsHitTestVisible = true;
 
+	UPROPERTY(BlueprintReadWrite, ReplicatedUsing = "OnRep_RenderOpacity", Getter, Setter, Category = "UI Framework", meta = (AllowPrivateAccess = "true"))
+	double RenderOpacity = 1.0;
+
 public:
 	//~ Begin UObject
 	virtual bool IsSupportedForNetworking() const override
@@ -149,6 +152,9 @@ public:
 	bool IsHitTestVisible() const;
 	void SetHitTestVisible(bool bInHitTestVisible);
 
+	double GetRenderOpacity() const;
+	void SetRenderOpacity(double InRenderOpacity);
+
 protected:
 	virtual void AuthorityRemoveChild(UUIFrameworkWidget* Widget)
 	{
@@ -167,6 +173,8 @@ private:
 	void OnRep_IsEnabled();
 	UFUNCTION()
 	void OnRep_Visibility();
+	UFUNCTION()
+	void OnRep_RenderOpacity();
 
 	void AuthoritySetWidgetTreeOwner(IUIFrameworkWidgetTreeOwner* InWidgetTreeOwner)
 	{
