UIFramework: Fix case of using a pointer to a removed element in a map that can lead to erroneous behaviour when Added/Removing widgets

diff --git a/Engine/Plugins/Experimental/UIFramework/Source/Private/Types/UIFWidgetTree.cpp b/Engine/Plugins/Experimental/UIFramework/Source/Private/Types/UIFWidgetTree.cpp
--- a/Engine/Plugins/Experimental/UIFramework/Source/Private/Types/UIFWidgetTree.cpp
+++ b/Engine/Plugins/Experimental/UIFramework/Source/Private/Types/UIFWidgetTree.cpp
@@ -339,11 +339,11 @@ bool FUIFrameworkWidgetTree::AuthorityRemoveChildRecursiveInternal(UUIFrameworkW
 
 	Widget->AuthoritySetWidgetTreeOwner(nullptr);
 	const FObjectKey WidgetKey = Widget;
-	if (int32* PreviousEntryIndexPtr = AuthorityIndexByWidgetMap.Find(WidgetKey))
+	int32 PreviousEntryIndex = INDEX_NONE;
+	if (AuthorityIndexByWidgetMap.RemoveAndCopyValue(WidgetKey, PreviousEntryIndex))
 	{
-		check(Entries.IsValidIndex(*PreviousEntryIndexPtr));
+		check(Entries.IsValidIndex(PreviousEntryIndex));
 
-		AuthorityIndexByWidgetMap.Remove(WidgetKey);
 		WidgetByIdMap.Remove(Widget->GetWidgetId());
 
 		if (ensure(ReplicatedOwner) && ReplicatedOwner->IsUsingRegisteredSubObjectList())
@@ -351,14 +351,14 @@ bool FUIFrameworkWidgetTree::AuthorityRemoveChildRecursiveInternal(UUIFrameworkW
 			ReplicatedOwner->RemoveReplicatedSubObject(Widget);
 		}
 
-		Entries.RemoveAtSwap(*PreviousEntryIndexPtr);
+		Entries.RemoveAtSwap(PreviousEntryIndex);
 
 		// Fix up the swap item
-		if (Entries.IsValidIndex(*PreviousEntryIndexPtr))
+		if (Entries.IsValidIndex(PreviousEntryIndex))
 		{
-			if (int32* FixUpChildIndexPtr = AuthorityIndexByWidgetMap.Find(FObjectKey(Entries[*PreviousEntryIndexPtr].Child)))
+			if (int32* FixUpChildIndexPtr = AuthorityIndexByWidgetMap.Find(FObjectKey(Entries[PreviousEntryIndex].Child)))
 			{
-				*FixUpChildIndexPtr = *PreviousEntryIndexPtr;
+				*FixUpChildIndexPtr = PreviousEntryIndex;
 			}
 		}
 
