UIFramework: Fix failures to Add/Remove widgets due to out-of-order replication

diff --git a/Engine/Plugins/Experimental/UIFramework/Source/Private/UIFPlayerComponent.cpp b/Engine/Plugins/Experimental/UIFramework/Source/Private/UIFPlayerComponent.cpp
--- a/Engine/Plugins/Experimental/UIFramework/Source/Private/UIFPlayerComponent.cpp
+++ b/Engine/Plugins/Experimental/UIFramework/Source/Private/UIFPlayerComponent.cpp
@@ -291,8 +291,12 @@ void UUIFrameworkPlayerComponent::LocalWidgetWasAddedToTree(const FUIFrameworkWi
 		{
 			if (Entry.IsParentValid() && Entry.IsChildValid())
 			{
-				Entry.Child->LocalCreateUMGWidget(this);
-				AddPending.Add(Entry.ReplicationID);
+				bool bDidCreateWidget = false;
+				Entry.Child->LocalCreateUMGWidget(this, &bDidCreateWidget);
+				if (bDidCreateWidget)
+				{
+					AddPending.Add(Entry.ReplicationID);
+				}
 				NetReplicationPending.Remove(Entry.ReplicationID);
 			}
 			else
@@ -379,8 +383,12 @@ void UUIFrameworkPlayerComponent::LocalOnClassLoaded(TSoftClassPtr<UWidget> Widg
 				{
 					if (Entry->IsParentValid() && Entry->IsChildValid())
 					{
-						Entry->Child->LocalCreateUMGWidget(this);
-						AddPending.Add(ReplicationId);
+						bool bDidCreateWidget = false;
+						Entry->Child->LocalCreateUMGWidget(this, &bDidCreateWidget);
+						if (bDidCreateWidget)
+						{
+							AddPending.Add(ReplicationId);
+						}
 						NetReplicationPending.Remove(ReplicationId);
 					}
 					else
diff --git a/Engine/Plugins/Experimental/UIFramework/Source/Private/UIFWidget.cpp b/Engine/Plugins/Experimental/UIFramework/Source/Private/UIFWidget.cpp
--- a/Engine/Plugins/Experimental/UIFramework/Source/Private/UIFWidget.cpp
+++ b/Engine/Plugins/Experimental/UIFramework/Source/Private/UIFWidget.cpp
@@ -99,44 +99,65 @@ void UUIFrameworkWidget::LocalAddChild(FUIFrameworkWidgetId ChildId)
 	}
 }
 
-void UUIFrameworkWidget::LocalCreateUMGWidget(TNonNullPtr<IUIFrameworkWidgetTreeOwner> InOwner)
+void UUIFrameworkWidget::LocalCreateUMGWidget(TNonNullPtr<IUIFrameworkWidgetTreeOwner> InOwner, bool* bDidCreateWidgetPtr)
 {
+	bool bIsNewWidget = false;
 	WidgetTreeOwner = InOwner;
 	if (UClass* Class = WidgetClass.Get())
 	{
-		if (Class->IsChildOf(UUserWidget::StaticClass()))
+		if (LocalUMGWidget && LocalUMGWidget->GetClass() != Class)
 		{
-			FUIFrameworkWidgetOwner UserWidgetOwner = WidgetTreeOwner->GetWidgetOwner();
-			if (UserWidgetOwner.PlayerController)
-			{
-				LocalUMGWidget = CreateWidget(UserWidgetOwner.PlayerController, Class);
-			}
-			else if (UserWidgetOwner.GameInstance)
-			{
-				LocalUMGWidget = CreateWidget(UserWidgetOwner.GameInstance, Class);
-			}
-			else if (UserWidgetOwner.World)
+			LocalUMGWidget->RemoveFromParent();
+			LocalUMGWidget = nullptr;
+		}
+
+		if (!LocalUMGWidget)
+		{
+			if (Class->IsChildOf(UUserWidget::StaticClass()))
 			{
-				LocalUMGWidget = CreateWidget(UserWidgetOwner.World, Class);
+				FUIFrameworkWidgetOwner UserWidgetOwner = WidgetTreeOwner->GetWidgetOwner();
+				if (UserWidgetOwner.PlayerController)
+				{
+					LocalUMGWidget = CreateWidget(UserWidgetOwner.PlayerController, Class);
+				}
+				else if (UserWidgetOwner.GameInstance)
+				{
+					LocalUMGWidget = CreateWidget(UserWidgetOwner.GameInstance, Class);
+				}
+				else if (UserWidgetOwner.World)
+				{
+					LocalUMGWidget = CreateWidget(UserWidgetOwner.World, Class);
+				}
+				else
+				{
+					ensureAlwaysMsgf(false, TEXT("There are no valid UserWidget owner."));
+				}
 			}
 			else
 			{
-				ensureAlwaysMsgf(false, TEXT("There are no valid UserWidget owner."));
+				check(Class->IsChildOf(UWidget::StaticClass()));
+				LocalUMGWidget = NewObject<UWidget>(this, Class, FName(), RF_Transient);
 			}
-		}
-		else
-		{
-			check(Class->IsChildOf(UWidget::StaticClass()));
-			LocalUMGWidget = NewObject<UWidget>(this, Class, FName(), RF_Transient);
+			
+			bIsNewWidget = LocalUMGWidget != nullptr;
 		}
 
 		if (LocalUMGWidget != nullptr)
 		{
 			LocalUMGWidget->SetIsEnabled(bIsEnabled);
 			LocalUMGWidget->SetVisibility(Visibility);
-			LocalOnUMGWidgetCreated();
+
+			if (bIsNewWidget)
+			{
+				LocalOnUMGWidgetCreated();
+			}
 		}
 	}
+
+	if (bDidCreateWidgetPtr)
+	{
+		*bDidCreateWidgetPtr = bIsNewWidget;
+	}
 }
 
 void UUIFrameworkWidget::LocalDestroyUMGWidget()
diff --git a/Engine/Plugins/Experimental/UIFramework/Source/Public/UIFWidget.h b/Engine/Plugins/Experimental/UIFramework/Source/Public/UIFWidget.h
--- a/Engine/Plugins/Experimental/UIFramework/Source/Public/UIFWidget.h
+++ b/Engine/Plugins/Experimental/UIFramework/Source/Public/UIFWidget.h
@@ -138,7 +138,7 @@ public:
 		return LocalUMGWidget;
 	}
 
-	void LocalCreateUMGWidget(TNonNullPtr<IUIFrameworkWidgetTreeOwner> Owner);
+	void LocalCreateUMGWidget(TNonNullPtr<IUIFrameworkWidgetTreeOwner> Owner, bool* bDidCreateWidgetPtr = nullptr);
 	virtual void LocalAddChild(FUIFrameworkWidgetId ChildId);
 	void LocalDestroyUMGWidget();
 
