UIFramework: Remove COND_InitialOnly from WidgetClass since it can be set after creation

#rnx

diff --git a/Engine/Plugins/Experimental/UIFramework/Source/Private/UIFWidget.cpp b/Engine/Plugins/Experimental/UIFramework/Source/Private/UIFWidget.cpp
--- a/Engine/Plugins/Experimental/UIFramework/Source/Private/UIFWidget.cpp
+++ b/Engine/Plugins/Experimental/UIFramework/Source/Private/UIFWidget.cpp
@@ -73,10 +73,8 @@ void UUIFrameworkWidget::GetLifetimeReplicatedProps(TArray<FLifetimeProperty>& O
 	DOREPLIFETIME_WITH_PARAMS_FAST(ThisClass, bIsEnabled, Params);
 	DOREPLIFETIME_WITH_PARAMS_FAST(ThisClass, Visibility, Params);
 	DOREPLIFETIME_WITH_PARAMS_FAST(ThisClass, RenderOpacity, Params);
-
-	Params.Condition = COND_InitialOnly;
-	DOREPLIFETIME_WITH_PARAMS_FAST(UUIFrameworkWidget, WidgetClass, Params);
-}
+	DOREPLIFETIME_WITH_PARAMS_FAST(ThisClass, WidgetClass, Params);
+} 
 
 FUIFrameworkWidgetTree* UUIFrameworkWidget::GetWidgetTree() const
 {
