UIFramework: Don't rely on FastArraySerializer's internal ItemMap as it's not always up to date when using Iris

#rnx

diff --git a/Engine/Plugins/Experimental/UIFramework/Source/Private/Types/UIFWidgetTree.cpp b/Engine/Plugins/Experimental/UIFramework/Source/Private/Types/UIFWidgetTree.cpp
--- a/Engine/Plugins/Experimental/UIFramework/Source/Private/Types/UIFWidgetTree.cpp
+++ b/Engine/Plugins/Experimental/UIFramework/Source/Private/Types/UIFWidgetTree.cpp
@@ -406,11 +406,10 @@ void FUIFrameworkWidgetTree::AuthorityRemoveAllWidgetsFromActorChannel()
 
 FUIFrameworkWidgetTreeEntry* FUIFrameworkWidgetTree::LocalGetEntryByReplicationId(int32 ReplicationId)
 {
-	if (const int32* Index = ItemMap.Find(ReplicationId))
+	return Entries.FindByPredicate([ReplicationId](const FUIFrameworkWidgetTreeEntry& Entry)
 	{
-		return &Entries[*Index];
-	}
-	return nullptr;
+		return Entry.ReplicationID == ReplicationId;
+	});
 }
 
 const FUIFrameworkWidgetTreeEntry* FUIFrameworkWidgetTree::LocalGetEntryByReplicationId(int32 ReplicationId) const
