UIFramework: Fix replicated widgets leaking on clients after the server has destroyed them

#rnx

diff --git a/Engine/Plugins/Experimental/UIFramework/Source/Private/Types/UIFWidgetTree.cpp b/Engine/Plugins/Experimental/UIFramework/Source/Private/Types/UIFWidgetTree.cpp
--- a/Engine/Plugins/Experimental/UIFramework/Source/Private/Types/UIFWidgetTree.cpp
+++ b/Engine/Plugins/Experimental/UIFramework/Source/Private/Types/UIFWidgetTree.cpp
@@ -337,7 +337,7 @@ bool FUIFrameworkWidgetTree::AuthorityRemoveChildRecursiveInternal(UUIFrameworkW
 
 		if (ensure(ReplicatedOwner) && ReplicatedOwner->IsUsingRegisteredSubObjectList())
 		{
-			ReplicatedOwner->RemoveReplicatedSubObject(Widget);
+			ReplicatedOwner->DestroyReplicatedSubObjectOnRemotePeers(Widget);
 		}
 
 		Entries.RemoveAtSwap(PreviousEntryIndex);
@@ -387,7 +387,7 @@ void FUIFrameworkWidgetTree::AuthorityRemoveAllWidgetsFromActorChannel()
 		{
 			if (Entry.Child && ReplicatedOwner->IsReplicatedSubObjectRegistered(Entry.Child))
 			{
-				ReplicatedOwner->RemoveReplicatedSubObject(Entry.Child);
+				ReplicatedOwner->DestroyReplicatedSubObjectOnRemotePeers(Entry.Child);
 			}
 		}
 	}
