UIFramework: allow focusing on a specific widget

#okforversepublic

diff --git a/Engine/Plugins/Experimental/UIFramework/Source/Private/Types/UIFWidgetTree.cpp b/Engine/Plugins/Experimental/UIFramework/Source/Private/Types/UIFWidgetTree.cpp
--- a/Engine/Plugins/Experimental/UIFramework/Source/Private/Types/UIFWidgetTree.cpp
+++ b/Engine/Plugins/Experimental/UIFramework/Source/Private/Types/UIFWidgetTree.cpp
@@ -417,6 +417,35 @@ const UUIFrameworkWidget* FUIFrameworkWidgetTree::FindWidgetById(FUIFrameworkWid
 	return const_cast<FUIFrameworkWidgetTree*>(this)->FindWidgetById(WidgetId);
 }
 
+const FUIFrameworkWidgetTreeEntry* FUIFrameworkWidgetTree::FindRootEntryById(FUIFrameworkWidgetId WidgetId) const
+{
+	FUIFrameworkWidgetId NextWidgetId = WidgetId;
+	while (NextWidgetId.IsValid())
+	{
+		const FUIFrameworkWidgetTreeEntry* Entry = Entries.FindByPredicate([NextWidgetId](const FUIFrameworkWidgetTreeEntry& Entry)
+			{
+				return Entry.ChildId == NextWidgetId;
+			});
+
+		if (Entry)
+		{
+			if (Entry->ParentId.IsRoot())
+			{
+				return Entry;
+			}
+			else
+			{
+				NextWidgetId = Entry->ParentId;
+			}
+		}
+		else
+		{
+			return nullptr;
+		}
+	}
+	return nullptr;
+}
+
 #if UE_UIFRAMEWORK_WITH_DEBUG
 void FUIFrameworkWidgetTree::AuthorityTest() const
 {
diff --git a/Engine/Plugins/Experimental/UIFramework/Source/Private/UIFPlayerComponent.cpp b/Engine/Plugins/Experimental/UIFramework/Source/Private/UIFPlayerComponent.cpp
--- a/Engine/Plugins/Experimental/UIFramework/Source/Private/UIFPlayerComponent.cpp
+++ b/Engine/Plugins/Experimental/UIFramework/Source/Private/UIFPlayerComponent.cpp
@@ -129,6 +129,7 @@ void UUIFrameworkPlayerComponent::UninitializeComponent()
 				ChildWidget->LocalDestroyUMGWidget();
 			}
 		}
+		FTSTicker::GetCoreTicker().RemoveTicker(PendingWidgetConstructedTickerHandle);
 	}
 
 	Super::UninitializeComponent();
@@ -142,6 +143,7 @@ void UUIFrameworkPlayerComponent::GetLifetimeReplicatedProps(TArray< FLifetimePr
 	Params.bIsPushBased = true;
 	DOREPLIFETIME_WITH_PARAMS_FAST(ThisClass, RootList, Params);
 	DOREPLIFETIME_WITH_PARAMS_FAST(ThisClass, WidgetTree, Params);
+	DOREPLIFETIME_WITH_PARAMS_FAST(ThisClass, WidgetToFocus, Params);
 }
 
 bool UUIFrameworkPlayerComponent::ReplicateSubobjects(UActorChannel* Channel, class FOutBunch* Bunch, FReplicationFlags* RepFlags)
@@ -257,6 +259,21 @@ void UUIFrameworkPlayerComponent::TickComponent(float DeltaTime, enum ELevelTick
 		OnPendingReplicationProcessed.Broadcast();
 
 		PrimaryComponentTick.SetTickFunctionEnable(false);
+
+		if (WidgetToFocus.IsValid())
+		{
+			if (UUIFrameworkWidget* Widget = GetWidgetTree().FindWidgetById(WidgetToFocus))
+			{
+				if (UWidget* UMGWidget = Widget->LocalGetUMGWidget())
+				{
+					if (!PendingWidgetConstructedTickerHandle.IsValid())
+					{
+						PendingWidgetConstructedTickerHandle = FTSTicker::GetCoreTicker().AddTicker(FTickerDelegate::CreateUObject(this, &UUIFrameworkPlayerComponent::TrySetFocus, WidgetToFocus));
+						WidgetToFocus = FUIFrameworkWidgetId();
+					}
+				}
+			}
+		}
 	}
 	else
 	{
@@ -370,6 +387,12 @@ void UUIFrameworkPlayerComponent::ServerRemoveWidgetRootFromTree_Implementation(
 	}
 }
 
+void UUIFrameworkPlayerComponent::SetWidgetToFocus(FUIFrameworkWidgetId WidgetId)
+{
+	WidgetToFocus = WidgetId;
+	MARK_PROPERTY_DIRTY_FROM_NAME(ThisClass, WidgetToFocus, this);
+}
+
 void UUIFrameworkPlayerComponent::LocalOnClassLoaded(TSoftClassPtr<UWidget> WidgetClass)
 {
 	FWidgetClassToLoad FoundWidgetClassToLoad;
@@ -446,3 +469,54 @@ void UUIFrameworkPlayerComponent::LocalAddChild(FUIFrameworkWidgetId WidgetId)
 		UE_LOG(LogUIFramework, Verbose, TEXT("The widget '%" INT64_FMT "' was not found in the RootList."), WidgetId.GetKey());
 	}
 }
+
+bool UUIFrameworkPlayerComponent::TrySetFocus(float DeltaTime, FUIFrameworkWidgetId CurrentWidgetToFocus)
+{
+	if (UUIFrameworkWidget* Widget = GetWidgetTree().FindWidgetById(CurrentWidgetToFocus))
+	{
+		if (UWidget* UMGWidget = Widget->LocalGetUMGWidget())
+		{
+			if (UMGWidget->IsConstructed())
+			{
+				const FUIFrameworkWidgetTreeEntry* RootEntry = GetWidgetTree().FindRootEntryById(CurrentWidgetToFocus);
+				if (RootEntry)
+				{
+					if (FUIFrameworkGameLayerSlot* Slot = RootList.FindEntry(RootEntry->ChildId))
+					{
+						if (Slot->InputMode == EUIFrameworkInputMode::UI)
+						{
+							UMGWidget->SetFocus();
+						}
+					}
+				}
+				FTSTicker::GetCoreTicker().RemoveTicker(PendingWidgetConstructedTickerHandle);
+			}
+		}
+	}
+	return true;
+}
+
+void UUIFrameworkPlayerComponent::OnRep_WidgetToFocus()
+{
+	if (WidgetToFocus.IsValid())
+	{
+		if (UUIFrameworkWidget* Widget = GetWidgetTree().FindWidgetById(WidgetToFocus))
+		{
+			if (UWidget* UMGWidget = Widget->LocalGetUMGWidget())
+			{
+				if (UMGWidget->IsConstructed())
+				{
+					TrySetFocus(0, WidgetToFocus);
+					WidgetToFocus = FUIFrameworkWidgetId();
+					return;
+				}
+			}
+			if (PendingWidgetConstructedTickerHandle.IsValid())
+			{
+				FTSTicker::GetCoreTicker().RemoveTicker(PendingWidgetConstructedTickerHandle);
+			}
+			PendingWidgetConstructedTickerHandle = FTSTicker::GetCoreTicker().AddTicker(FTickerDelegate::CreateUObject(this, &UUIFrameworkPlayerComponent::TrySetFocus, WidgetToFocus));
+			WidgetToFocus = FUIFrameworkWidgetId();
+		}
+	}
+}
\ No newline at end of file
diff --git a/Engine/Plugins/Experimental/UIFramework/Source/Public/Types/UIFWidgetTree.h b/Engine/Plugins/Experimental/UIFramework/Source/Public/Types/UIFWidgetTree.h
--- a/Engine/Plugins/Experimental/UIFramework/Source/Public/Types/UIFWidgetTree.h
+++ b/Engine/Plugins/Experimental/UIFramework/Source/Public/Types/UIFWidgetTree.h
@@ -122,6 +122,9 @@ public:
 	/** Removes all widgets added to the ActorChannel replicated list */
 	void AuthorityRemoveAllWidgetsFromActorChannel();
 
+	/** Gets the root widget of the tree given the Id of a widget in the tree */
+	const FUIFrameworkWidgetTreeEntry* FindRootEntryById(FUIFrameworkWidgetId WidgetId) const;
+
 #if UE_UIFRAMEWORK_WITH_DEBUG
 	void AuthorityTest() const;
 	void LogTree() const;
diff --git a/Engine/Plugins/Experimental/UIFramework/Source/Public/UIFPlayerComponent.h b/Engine/Plugins/Experimental/UIFramework/Source/Public/UIFPlayerComponent.h
--- a/Engine/Plugins/Experimental/UIFramework/Source/Public/UIFPlayerComponent.h
+++ b/Engine/Plugins/Experimental/UIFramework/Source/Public/UIFPlayerComponent.h
@@ -2,6 +2,7 @@
 
 #pragma once
 
+#include "Containers/Ticker.h"
 #include "GameFramework/PlayerController.h"
 #include "Types/UIFSlotBase.h"
 #include "Types/UIFWidgetTree.h"
@@ -174,6 +175,8 @@ public:
 	virtual void LocalWidgetRemovedFromTree(const FUIFrameworkWidgetTreeEntry& Entry) override;
 	virtual void LocalRemoveWidgetRootFromTree(const UUIFrameworkWidget* Widget) override;
 
+	void SetWidgetToFocus(FUIFrameworkWidgetId WidgetId);
+
 private:
 	UFUNCTION(Server, Reliable)
 	void ServerRemoveWidgetRootFromTree(FUIFrameworkWidgetId WidgetId);
@@ -181,6 +184,10 @@ private:
 	void LocalOnClassLoaded(TSoftClassPtr<UWidget> WidgetClass);
 	void LocalAddChild(FUIFrameworkWidgetId WidgetId);
 
+	UFUNCTION()
+	void OnRep_WidgetToFocus();
+	bool TrySetFocus(float DeltaTime, FUIFrameworkWidgetId CurrentWidgetToFocus);
+
 private:
 	UPROPERTY(Replicated)
 	FUIFrameworkGameLayerSlotList RootList;
@@ -202,6 +209,10 @@ private:
 	//~ Once widgets are created and constructed, allow for actions such as focus to occur
 	FOnPendingReplicationProcessed OnPendingReplicationProcessed;
 
+	//~ Widget ID to focus on
+	UPROPERTY(ReplicatedUsing=OnRep_WidgetToFocus, Transient)
+	FUIFrameworkWidgetId WidgetToFocus;
+
 	struct FWidgetClassToLoad
 	{
 		TArray<int32, TInlineAllocator<4>> EntryReplicationIds;
@@ -211,4 +222,6 @@ private:
 	TMap<TSoftClassPtr<UWidget>, FWidgetClassToLoad> ClassesToLoad;
 
 	bool bAddingWidget = false;
+
+	FTSTicker::FDelegateHandle PendingWidgetConstructedTickerHandle;
 };
