UIFramework: Add LocalGetOrCreateUMGWidgetIfReady() helper to facilitate waiting for widgets to replicate and load

#rnx

diff --git a/Engine/Plugins/Experimental/UIFramework/Source/Private/UIFWidget.cpp b/Engine/Plugins/Experimental/UIFramework/Source/Private/UIFWidget.cpp
--- a/Engine/Plugins/Experimental/UIFramework/Source/Private/UIFWidget.cpp
+++ b/Engine/Plugins/Experimental/UIFramework/Source/Private/UIFWidget.cpp
@@ -4,6 +4,7 @@
 //#include "UIFManagerSubsystem.h"
 
 #include "Blueprint/UserWidget.h"
+#include "Engine/AssetManager.h"
 
 #include "Engine/Engine.h"
 #include "Engine/NetDriver.h"
@@ -97,7 +98,71 @@ void UUIFrameworkWidget::LocalAddChild(FUIFrameworkWidgetId ChildId)
 	}
 }
 
-void UUIFrameworkWidget::LocalCreateUMGWidget(TNonNullPtr<IUIFrameworkWidgetTreeOwner> InOwner, bool* bDidCreateWidgetPtr)
+UWidget* UUIFrameworkWidget::LocalGetOrCreateUMGWidgetIfReady()
+{
+	if (UWidget* Widget = LocalGetUMGWidget())
+	{
+		if (Widget->GetClass() == WidgetClass.Get())
+		{
+			return Widget;
+		}
+	}
+	
+	if (!LocalIsReplicationReady())
+	{
+		return nullptr;
+	}
+	
+	if (WidgetClass.IsNull())
+	{
+		return nullptr;
+	}
+
+	if (WidgetClass.IsPending())
+	{
+		AsyncLoadWidgetClass();
+		return nullptr;
+	}
+
+	LocalCreateUMGWidget(nullptr, nullptr);
+	return LocalGetUMGWidget();
+}
+
+TSharedPtr<FStreamableHandle> UUIFrameworkWidget::AsyncLoadWidgetClass()
+{
+	if (WidgetClassStreamableHandle.IsValid())
+	{
+		// Check that we're loading the correct class
+		TArray<FSoftObjectPath> LoadingWidgets;
+		constexpr bool bIncludeChildHandles = false;
+		WidgetClassStreamableHandle->GetRequestedAssets(LoadingWidgets, bIncludeChildHandles);
+		if (LoadingWidgets.Contains(WidgetClass.ToSoftObjectPath()))
+		{
+			return WidgetClassStreamableHandle;
+		}
+
+		// WidgetClass changed while we were loading
+		WidgetClassStreamableHandle->CancelHandle();
+		WidgetClassStreamableHandle.Reset();
+	}
+	
+	TWeakObjectPtr<ThisClass> WeakSelf = this;
+	WidgetClassStreamableHandle = UAssetManager::GetStreamableManager().RequestAsyncLoad(
+		WidgetClass.ToSoftObjectPath()
+		, [WeakSelf]()
+		{
+			if (ThisClass* StrongSelf = WeakSelf.Get())
+			{
+				StrongSelf->LocalCreateUMGWidget(nullptr, nullptr);
+				StrongSelf->WidgetClassStreamableHandle.Reset();
+			}
+		}
+	, FStreamableManager::AsyncLoadHighPriority, false, false, TEXT("UUIFrameworkWidget::WidgetClass"));
+
+	return WidgetClassStreamableHandle;
+}
+
+void UUIFrameworkWidget::LocalCreateUMGWidget(IUIFrameworkWidgetTreeOwner* InOwner, bool* bDidCreateWidgetPtr)
 {
 	bool bIsNewWidget = false;
 	WidgetTreeOwner = InOwner;
@@ -113,7 +178,7 @@ void UUIFrameworkWidget::LocalCreateUMGWidget(TNonNullPtr<IUIFrameworkWidgetTree
 		{
 			if (Class->IsChildOf(UUserWidget::StaticClass()))
 			{
-				FUIFrameworkWidgetOwner UserWidgetOwner = WidgetTreeOwner->GetWidgetOwner();
+				FUIFrameworkWidgetOwner UserWidgetOwner = WidgetTreeOwner ? WidgetTreeOwner->GetWidgetOwner() : FUIFrameworkWidgetOwner{ GetPlayerController() };
 				if (UserWidgetOwner.PlayerController)
 				{
 					LocalUMGWidget = CreateWidget(UserWidgetOwner.PlayerController, Class);
diff --git a/Engine/Plugins/Experimental/UIFramework/Source/Public/UIFWidget.h b/Engine/Plugins/Experimental/UIFramework/Source/Public/UIFWidget.h
--- a/Engine/Plugins/Experimental/UIFramework/Source/Public/UIFWidget.h
+++ b/Engine/Plugins/Experimental/UIFramework/Source/Public/UIFWidget.h
@@ -9,6 +9,7 @@
 
 #include "UIFWidget.generated.h"
 
+struct FStreamableHandle;
 template <typename ObjectType> class TNonNullPtr;
 
 class FUIFrameworkModule;
@@ -138,7 +139,12 @@ public:
 		return LocalUMGWidget;
 	}
 
-	void LocalCreateUMGWidget(TNonNullPtr<IUIFrameworkWidgetTreeOwner> Owner, bool* bDidCreateWidgetPtr = nullptr);
+	// Can be used instead of manually calling LocalIsReplicationReady() and AsyncLoadWidgetClass()
+	UWidget* LocalGetOrCreateUMGWidgetIfReady();
+
+	TSharedPtr<FStreamableHandle> AsyncLoadWidgetClass();
+
+	void LocalCreateUMGWidget(IUIFrameworkWidgetTreeOwner* Owner, bool* bDidCreateWidgetPtr = nullptr);
 	virtual void LocalAddChild(FUIFrameworkWidgetId ChildId);
 	void LocalDestroyUMGWidget();
 
@@ -208,4 +214,6 @@ private:
 	//~ LocalOnly
 	UPROPERTY(Transient)
 	TObjectPtr<UWidget> LocalUMGWidget;
+
+	TSharedPtr<FStreamableHandle> WidgetClassStreamableHandle;
 };
