[Fail] UI fails to load properly after flushing into multiple islands

#rnx

diff --git a/Engine/Plugins/Experimental/UIFramework/Source/Private/UIFWidget.cpp b/Engine/Plugins/Experimental/UIFramework/Source/Private/UIFWidget.cpp
--- a/Engine/Plugins/Experimental/UIFramework/Source/Private/UIFWidget.cpp
+++ b/Engine/Plugins/Experimental/UIFramework/Source/Private/UIFWidget.cpp
@@ -17,6 +17,15 @@
 #include UE_INLINE_GENERATED_CPP_BY_NAME(UIFWidget)
 
 
+UUIFrameworkWidget::UUIFrameworkWidget(const FObjectInitializer& ObjectInitializer)
+: Super(ObjectInitializer)
+{
+	if (!IsTemplate())
+	{
+		Id = FUIFrameworkWidgetId::MakeNew();
+	}
+}
+
 void UUIFrameworkWidget::ForceNetUpdate()
 {
 	if (AActor* OwnerActor = Cast<AActor>(GetOuter()))
diff --git a/Engine/Plugins/Experimental/UIFramework/Source/Public/UIFWidget.h b/Engine/Plugins/Experimental/UIFramework/Source/Public/UIFWidget.h
--- a/Engine/Plugins/Experimental/UIFramework/Source/Public/UIFWidget.h
+++ b/Engine/Plugins/Experimental/UIFramework/Source/Public/UIFWidget.h
@@ -12,10 +12,11 @@
 struct FStreamableHandle;
 template <typename ObjectType> class TNonNullPtr;
 
+class FObjectInitializer;
 class FUIFrameworkModule;
+class IUIFrameworkWidgetTreeOwner;
 class UUIFrameworkWidget;
 struct FUIFrameworkWidgetTree;
-class IUIFrameworkWidgetTreeOwner;
 
 /**
  *
@@ -42,6 +43,10 @@ class UIFRAMEWORK_API UUIFrameworkWidget : public UMVVMViewModelBase
 	friend FUIFrameworkModule;
 	friend FUIFrameworkWidgetTree;
 
+public:
+	UUIFrameworkWidget() = default;
+	UUIFrameworkWidget(const FObjectInitializer& ObjectInitializer);
+
 private:
 	UPROPERTY(BlueprintReadWrite, ReplicatedUsing = "OnRep_IsEnabled", Getter = "IsEnabled", Setter="SetEnabled", Category = "UI Framework", meta = (AllowPrivateAccess = "true"))
 	bool bIsEnabled = true;
@@ -198,7 +203,7 @@ protected:
 private:
 	//~ Authority and Local
 	UPROPERTY(Replicated, Transient, DuplicateTransient)
-	FUIFrameworkWidgetId Id = FUIFrameworkWidgetId::MakeNew();
+	FUIFrameworkWidgetId Id;
 
 	//~ Authority
 	UPROPERTY(Transient)
