[Fix] Verse UI spawned widgets were not properly constructed and attaching to UIFrameworkPlayerComponent

diff --git a/Engine/Plugins/Experimental/UIFramework/Source/Private/UIFWidget.cpp b/Engine/Plugins/Experimental/UIFramework/Source/Private/UIFWidget.cpp
--- a/Engine/Plugins/Experimental/UIFramework/Source/Private/UIFWidget.cpp
+++ b/Engine/Plugins/Experimental/UIFramework/Source/Private/UIFWidget.cpp
@@ -16,6 +16,14 @@
 
 #include UE_INLINE_GENERATED_CPP_BY_NAME(UIFWidget)
 
+UUIFrameworkWidget::UUIFrameworkWidget()
+: Super()
+{
+	if (!IsTemplate())
+	{
+		Id = FUIFrameworkWidgetId::MakeNew();
+	}
+}
 
 UUIFrameworkWidget::UUIFrameworkWidget(const FObjectInitializer& ObjectInitializer)
 : Super(ObjectInitializer)
diff --git a/Engine/Plugins/Experimental/UIFramework/Source/Public/UIFWidget.h b/Engine/Plugins/Experimental/UIFramework/Source/Public/UIFWidget.h
--- a/Engine/Plugins/Experimental/UIFramework/Source/Public/UIFWidget.h
+++ b/Engine/Plugins/Experimental/UIFramework/Source/Public/UIFWidget.h
@@ -44,7 +44,7 @@ class UIFRAMEWORK_API UUIFrameworkWidget : public UMVVMViewModelBase
 	friend FUIFrameworkWidgetTree;
 
 public:
-	UUIFrameworkWidget() = default;
+	UUIFrameworkWidget();
 	UUIFrameworkWidget(const FObjectInitializer& ObjectInitializer);
 
 private:
