Scene Graph UI: add a component to display UI
#okforversepublic

#virtualized

diff --git a/Engine/Plugins/Experimental/UIFramework/Source/Public/UIFPlayerComponent.h b/Engine/Plugins/Experimental/UIFramework/Source/Public/UIFPlayerComponent.h
--- a/Engine/Plugins/Experimental/UIFramework/Source/Public/UIFPlayerComponent.h
+++ b/Engine/Plugins/Experimental/UIFramework/Source/Public/UIFPlayerComponent.h
@@ -179,12 +179,14 @@ public:
 
 	UE_API void SetWidgetToFocus(FUIFrameworkWidgetId WidgetId);
 
+protected:
+	UE_API virtual void LocalAddChild(FUIFrameworkWidgetId WidgetId);
+
 private:
 	UFUNCTION(Server, Reliable)
 	UE_API void ServerRemoveWidgetRootFromTree(FUIFrameworkWidgetId WidgetId);
 
 	UE_API void LocalOnClassLoaded(TSoftClassPtr<UWidget> WidgetClass);
-	UE_API void LocalAddChild(FUIFrameworkWidgetId WidgetId);
 
 	UFUNCTION()
 	UE_API void OnRep_WidgetToFocus();
diff --git a/Engine/Source/Runtime/UMG/Private/Components/WidgetComponent.cpp b/Engine/Source/Runtime/UMG/Private/Components/WidgetComponent.cpp
--- a/Engine/Source/Runtime/UMG/Private/Components/WidgetComponent.cpp
+++ b/Engine/Source/Runtime/UMG/Private/Components/WidgetComponent.cpp
@@ -788,6 +788,8 @@ void UWidgetComponent::UpdateMaterialInstance()
 
 	PrecachePSOs();
 
+	OnMaterialInstanceUpdated.Broadcast();
+
 	MarkRenderStateDirty();
 }
 
diff --git a/Engine/Source/Runtime/UMG/Public/Components/WidgetComponent.h b/Engine/Source/Runtime/UMG/Public/Components/WidgetComponent.h
--- a/Engine/Source/Runtime/UMG/Public/Components/WidgetComponent.h
+++ b/Engine/Source/Runtime/UMG/Public/Components/WidgetComponent.h
@@ -409,6 +409,9 @@ public:
 	/** Hook to allow this component modify the local position of the widget after it has been projected from world space to screen space. */
 	virtual FVector2D ModifyProjectedLocalPosition(const FGeometry& ViewportGeometry, const FVector2D& LocalPosition) { return LocalPosition; }
 
+	DECLARE_MULTICAST_DELEGATE(FOnMaterialInstanceUpdated);
+	FOnMaterialInstanceUpdated OnMaterialInstanceUpdated;
+
 protected:
 	UMG_API void OnLevelRemovedFromWorld(ULevel* InLevel, UWorld* InWorld);
 
