Fix one frame delay when adding widgets

https://p4-swarm.epicgames.net/reviews/41975328

#rnx
#tests Tested removing/adding widgets every 1/0.1/0.01s with and without the fix

diff --git a/Engine/Plugins/Experimental/UIFramework/Source/Private/UIFPlayerComponent.cpp b/Engine/Plugins/Experimental/UIFramework/Source/Private/UIFPlayerComponent.cpp
--- a/Engine/Plugins/Experimental/UIFramework/Source/Private/UIFPlayerComponent.cpp
+++ b/Engine/Plugins/Experimental/UIFramework/Source/Private/UIFPlayerComponent.cpp
@@ -221,61 +221,64 @@ FUIFrameworkWidgetOwner UUIFrameworkPlayerComponent::GetWidgetOwner() const
 	return Owner;
 }
 
-void UUIFrameworkPlayerComponent::TickComponent(float DeltaTime, enum ELevelTick TickType, FActorComponentTickFunction* ThisTickFunction)
+void UUIFrameworkPlayerComponent::HandleAddPending()
 {
-	APlayerController* LocalOwner = GetPlayerController();
-	check(!LocalOwner || !LocalOwner->HasAuthority());
-
-	if (ClassesToLoad.Num() == 0 && NetReplicationPending.Num() == 0)
+	ensureMsgf(AddPending.Num() != 0, TEXT("we shouldn't call HandleAddPending when nothing needs to be added"));
+	// Create and add all the pending widgets
+	TGuardValue<bool> TmpGuard = { bAddingWidget, true };
+	for (int32 ReplicationId : AddPending)
 	{
-		// Create and add all the pending widgets
-		TGuardValue<bool> TmpGuard = {bAddingWidget, true};
-		for (int32 ReplicationId : AddPending)
+		if (FUIFrameworkWidgetTreeEntry* Entry = WidgetTree.LocalGetEntryByReplicationId(ReplicationId))
 		{
-			if (FUIFrameworkWidgetTreeEntry* Entry = WidgetTree.LocalGetEntryByReplicationId(ReplicationId))
+			if (ensure(Entry->IsParentValid() && Entry->IsChildValid()))
 			{
-				if (ensure(Entry->IsParentValid() && Entry->IsChildValid()))
+				if (Entry->ParentId.IsRoot())
 				{
-					if (Entry->ParentId.IsRoot())
-					{
-						LocalAddChild(Entry->ChildId);
-					}
-					else
-					{
-						Entry->Parent->LocalAddChild(Entry->ChildId);
-					}
+					LocalAddChild(Entry->ChildId);
+				}
+				else
+				{
+					Entry->Parent->LocalAddChild(Entry->ChildId);
 				}
-			}
-			else
-			{
-				UE_LOG(LogUIFramework, Verbose, TEXT("A widget would was added but couldn't be found anymore."));
 			}
 		}
+		else
+		{
+			UE_LOG(LogUIFramework, Verbose, TEXT("A widget would have been added but couldn't be found anymore."));
+		}
+	}
 
-		NetReplicationPending.Empty();
-		AddPending.Empty();
-		ClassesToLoad.Empty();
+	ensureMsgf(NetReplicationPending.Num() == 0, TEXT("we shouldn't be in HandleAddPending if we're waiting on net replication for some of the objects"));
+	ensureMsgf(ClassesToLoad.Num() == 0, TEXT("we shouldn't be in HandleAddPending if our classes haven't been loaded yet"));
+	AddPending.Empty();
 
-		OnPendingReplicationProcessed.Broadcast();
+	OnPendingReplicationProcessed.Broadcast();
 
-		PrimaryComponentTick.SetTickFunctionEnable(false);
+	PrimaryComponentTick.SetTickFunctionEnable(false);
 
-		if (WidgetToFocus.IsValid())
+	if (WidgetToFocus.IsValid())
+	{
+		if (UUIFrameworkWidget* Widget = GetWidgetTree().FindWidgetById(WidgetToFocus))
 		{
-			if (UUIFrameworkWidget* Widget = GetWidgetTree().FindWidgetById(WidgetToFocus))
+			if (UWidget* UMGWidget = Widget->LocalGetUMGWidget())
 			{
-				if (UWidget* UMGWidget = Widget->LocalGetUMGWidget())
+				if (!PendingWidgetConstructedTickerHandle.IsValid())
 				{
-					if (!PendingWidgetConstructedTickerHandle.IsValid())
-					{
-						PendingWidgetConstructedTickerHandle = FTSTicker::GetCoreTicker().AddTicker(FTickerDelegate::CreateUObject(this, &UUIFrameworkPlayerComponent::TrySetFocus, WidgetToFocus));
-						WidgetToFocus = FUIFrameworkWidgetId();
-					}
+					PendingWidgetConstructedTickerHandle = FTSTicker::GetCoreTicker().AddTicker(FTickerDelegate::CreateUObject(this, &UUIFrameworkPlayerComponent::TrySetFocus, WidgetToFocus));
+					WidgetToFocus = FUIFrameworkWidgetId();
 				}
 			}
 		}
 	}
-	else
+
+}
+
+void UUIFrameworkPlayerComponent::TickComponent(float DeltaTime, enum ELevelTick TickType, FActorComponentTickFunction* ThisTickFunction)
+{
+	APlayerController* LocalOwner = GetPlayerController();
+	check(!LocalOwner || !LocalOwner->HasAuthority());
+
+	if (ClassesToLoad.Num() != 0 || NetReplicationPending.Num() != 0)
 	{
 		FMemMark Mark(FMemStack::Get());
 		TArray<int32, TMemStackAllocator<>> TempNetReplicationPending;
@@ -294,6 +297,12 @@ void UUIFrameworkPlayerComponent::TickComponent(float DeltaTime, enum ELevelTick
 			}
 		}
 	}
+
+	if (AddPending.Num() > 0 && ClassesToLoad.Num() == 0 && NetReplicationPending.Num() == 0)
+	{
+		HandleAddPending();
+	}
+
 }
 
 void UUIFrameworkPlayerComponent::LocalWidgetWasAddedToTree(const FUIFrameworkWidgetTreeEntry& Entry)
diff --git a/Engine/Plugins/Experimental/UIFramework/Source/Public/UIFPlayerComponent.h b/Engine/Plugins/Experimental/UIFramework/Source/Public/UIFPlayerComponent.h
--- a/Engine/Plugins/Experimental/UIFramework/Source/Public/UIFPlayerComponent.h
+++ b/Engine/Plugins/Experimental/UIFramework/Source/Public/UIFPlayerComponent.h
@@ -192,6 +192,7 @@ private:
 	UE_API void OnRep_WidgetToFocus();
 	UE_API bool TrySetFocus(float DeltaTime, FUIFrameworkWidgetId CurrentWidgetToFocus);
 
+	void HandleAddPending();
 private:
 	UPROPERTY(Replicated)
 	FUIFrameworkGameLayerSlotList RootList;
