Add invalidation to SEditableTextBox style
expose text size
#okforversepublic

diff --git a/Engine/Plugins/Experimental/UIFramework/Source/Private/Widgets/UIFTextBlock.cpp b/Engine/Plugins/Experimental/UIFramework/Source/Private/Widgets/UIFTextBlock.cpp
--- a/Engine/Plugins/Experimental/UIFramework/Source/Private/Widgets/UIFTextBlock.cpp
+++ b/Engine/Plugins/Experimental/UIFramework/Source/Private/Widgets/UIFTextBlock.cpp
@@ -64,6 +64,16 @@ void UUIFrameworkTextBase::SetTextColor(FLinearColor InTextColor)
 	}
 }
 
+void UUIFrameworkTextBase::SetTextSize(float InTextSize)
+{
+	if (TextSize != InTextSize)
+	{
+		TextSize = InTextSize;
+		MARK_PROPERTY_DIRTY_FROM_NAME(ThisClass, TextSize, this);
+		ForceNetUpdate();
+	}
+}
+
 void UUIFrameworkTextBase::SetJustification(ETextJustify::Type InJustification)
 {
 	if (Justification != InJustification)
@@ -89,6 +99,7 @@ void UUIFrameworkTextBase::LocalOnUMGWidgetCreated()
 	SetText(Message);
 	SetTextToWidget(Text);
 	SetTextColorToWidget(TextColor);
+	SetTextSizeToWidget(TextSize);
 	SetJustificationToWidget(Justification);
 	SetOverflowPolicyToWidget(OverflowPolicy);
 }
@@ -111,6 +122,14 @@ void UUIFrameworkTextBase::OnRep_TextColor()
 	}
 }
 
+void UUIFrameworkTextBase::OnRep_TextSize()
+{
+	if (LocalGetUMGWidget())
+	{
+		SetTextSizeToWidget(TextSize);
+	}
+}
+
 void UUIFrameworkTextBase::OnRep_Justification()
 {
 	if (LocalGetUMGWidget())
@@ -193,6 +212,14 @@ void UUIFrameworkTextBlock::SetTextToWidget(const FText& InValue)
 	CastChecked<UTextBlock>(LocalGetUMGWidget())->SetText(InValue);
 }
 
+void UUIFrameworkTextBlock::SetTextSizeToWidget(float InValue)
+{
+	UTextBlock* TextBlock = CastChecked<UTextBlock>(LocalGetUMGWidget());
+	FSlateFontInfo CurrentFontInfo = TextBlock->GetFont();
+	CurrentFontInfo.Size = InValue;
+	TextBlock->SetFont(CurrentFontInfo);
+}
+
 void UUIFrameworkTextBlock::SetTextColorToWidget(FLinearColor InValue)
 {
 	CastChecked<UTextBlock>(LocalGetUMGWidget())->SetColorAndOpacity(InValue);
diff --git a/Engine/Plugins/Experimental/UIFramework/Source/Public/Widgets/UIFTextBlock.h b/Engine/Plugins/Experimental/UIFramework/Source/Public/Widgets/UIFTextBlock.h
--- a/Engine/Plugins/Experimental/UIFramework/Source/Public/Widgets/UIFTextBlock.h
+++ b/Engine/Plugins/Experimental/UIFramework/Source/Public/Widgets/UIFTextBlock.h
@@ -43,6 +43,15 @@ public:
 		return TextColor;
 	}
 
+	UFUNCTION(BlueprintCallable, BlueprintAuthorityOnly, Category = "UI Framework")
+	UE_API void SetTextSize(float TextSize);
+
+	UFUNCTION(BlueprintCallable, Category = "UI Framework")
+	float GetTextSize() const
+	{
+		return TextSize;
+	}
+
 	UFUNCTION(BlueprintCallable, BlueprintAuthorityOnly, Category = "UI Framework")
 	UE_API void SetJustification(ETextJustify::Type Justification);
 
@@ -70,6 +79,9 @@ private:
 	UFUNCTION()
 	UE_API void OnRep_TextColor();
 
+	UFUNCTION()
+	UE_API void OnRep_TextSize();
+
 	UFUNCTION()
 	UE_API void OnRep_Justification();
 
@@ -85,6 +97,9 @@ protected:
 	virtual void SetTextColorToWidget(FLinearColor)
 	{
 	}
+	virtual void SetTextSizeToWidget(float)
+	{
+	}
 	virtual void SetJustificationToWidget(ETextJustify::Type)
 	{
 	}
@@ -102,6 +117,9 @@ private:
 	UPROPERTY(ReplicatedUsing = OnRep_TextColor)
 	FLinearColor TextColor = FLinearColor::Black;
 
+	UPROPERTY(ReplicatedUsing = OnRep_TextSize)
+	float TextSize = 24.0;
+
 	UPROPERTY(ReplicatedUsing = OnRep_Justification)
 	TEnumAsByte<ETextJustify::Type> Justification;
 
@@ -151,6 +169,7 @@ protected:
 	UE_API virtual void SetJustificationToWidget(ETextJustify::Type) override;
 	UE_API virtual void SetOverflowPolicyToWidget(ETextOverflowPolicy) override;
 	UE_API virtual void SetTextColorToWidget(FLinearColor) override;
+	UE_API virtual void SetTextSizeToWidget(float) override;
 	UE_API virtual void SetShadowOffsetToWidget(FVector2f);
 	UE_API virtual void SetShadowColorToWidget(FLinearColor);
 
diff --git a/Engine/Source/Runtime/Slate/Private/Widgets/Input/SEditableTextBox.cpp b/Engine/Source/Runtime/Slate/Private/Widgets/Input/SEditableTextBox.cpp
--- a/Engine/Source/Runtime/Slate/Private/Widgets/Input/SEditableTextBox.cpp
+++ b/Engine/Source/Runtime/Slate/Private/Widgets/Input/SEditableTextBox.cpp
@@ -131,6 +131,7 @@ void SEditableTextBox::SetStyle(const FEditableTextBoxStyle* InStyle)
 	BorderImageReadOnly = &Style->BackgroundImageReadOnly;
 
 	SetTextBlockStyle(&Style->TextStyle);
+	Invalidate(EInvalidateWidgetReason::Layout);
 }
 
 void SEditableTextBox::SetTextBlockStyle(const FTextBlockStyle* InTextStyle)
diff --git a/Engine/Source/Runtime/UMG/Private/Components/EditableTextBox.cpp b/Engine/Source/Runtime/UMG/Private/Components/EditableTextBox.cpp
--- a/Engine/Source/Runtime/UMG/Private/Components/EditableTextBox.cpp
+++ b/Engine/Source/Runtime/UMG/Private/Components/EditableTextBox.cpp
@@ -374,6 +374,20 @@ void UEditableTextBox::SetTextOverflowPolicy(ETextOverflowPolicy InOverflowPolic
 	}
 }
 
+void UEditableTextBox::SetWidgetStyle(const FEditableTextBoxStyle& InStyle)
+{
+	WidgetStyle = InStyle;
+	if (MyEditableTextBlock.IsValid())
+	{
+		MyEditableTextBlock->SetStyle(&InStyle);
+	}
+}
+
+const FEditableTextBoxStyle& UEditableTextBox::GetWidgetStyle() const
+{
+	return WidgetStyle;
+}
+
 void UEditableTextBox::HandleOnTextChanged(const FText& InText)
 {
 	if (SetTextInternal(InText))
diff --git a/Engine/Source/Runtime/UMG/Public/Components/EditableTextBox.h b/Engine/Source/Runtime/UMG/Public/Components/EditableTextBox.h
--- a/Engine/Source/Runtime/UMG/Public/Components/EditableTextBox.h
+++ b/Engine/Source/Runtime/UMG/Public/Components/EditableTextBox.h
@@ -45,7 +45,7 @@ public:
 public:
 
 	/** The style */
-	UPROPERTY(EditAnywhere, BlueprintReadWrite, Category=Appearance, meta=(DisplayName="Style", ShowOnlyInnerProperties))
+	UPROPERTY(EditAnywhere, BlueprintReadWrite, Getter, Setter, Category=Appearance, meta=(DisplayName="Style", ShowOnlyInnerProperties))
 	FEditableTextBoxStyle WidgetStyle;
 
 	UE_DEPRECATED(5.1, "Direct access to Hint Text is deprecated. Please use the getter or setter.")
@@ -249,6 +249,10 @@ public:
 	UFUNCTION(BlueprintSetter)
 	UMG_API void SetTextOverflowPolicy(ETextOverflowPolicy InOverflowPolicy);
 
+	UMG_API void SetWidgetStyle(const FEditableTextBoxStyle& InStyle);
+
+	UMG_API const FEditableTextBoxStyle& GetWidgetStyle() const;
+
 	UFUNCTION(BlueprintCallable, Category="Widget", meta=(DisplayName="SetForegroundColor (Text Box)"))
 	UMG_API void SetForegroundColor(FLinearColor color);
 
