VerseUI: Expose SetFocus method to player_ui.

[RN] minor creator/ui
#okforversepublic

diff --git a/Engine/Plugins/Experimental/UIFramework/Source/Private/UIFPlayerComponent.cpp b/Engine/Plugins/Experimental/UIFramework/Source/Private/UIFPlayerComponent.cpp
--- a/Engine/Plugins/Experimental/UIFramework/Source/Private/UIFPlayerComponent.cpp
+++ b/Engine/Plugins/Experimental/UIFramework/Source/Private/UIFPlayerComponent.cpp
@@ -451,7 +451,7 @@ void UUIFrameworkPlayerComponent::LocalOnClassLoaded(TSoftClassPtr<UWidget> Widg
 void UUIFrameworkPlayerComponent::LocalAddChild(FUIFrameworkWidgetId WidgetId)
 {
 	UUIFrameworkWidget* Widget = GetWidgetTree().FindWidgetById(WidgetId);
-	if (FUIFrameworkGameLayerSlot* LayerEntry = RootList.FindEntry(WidgetId))
+	if (FUIFrameworkGameLayerSlot* LayerSlot = RootList.FindEntry(WidgetId))
 	{
 		if (Widget)
 		{
@@ -459,9 +459,9 @@ void UUIFrameworkPlayerComponent::LocalAddChild(FUIFrameworkWidgetId WidgetId)
 			if (ensure(UMGWidget))
 			{
 				UMGWidget->RemoveFromParent();
-				LayerEntry->LocalAquireWidget();
+				LayerSlot->LocalAquireWidget();
 				check(Presenter);
-				Presenter->AddToViewport(UMGWidget, *LayerEntry);
+				Presenter->AddToViewport(UMGWidget, *LayerSlot);
 			}
 		}
 		else
@@ -490,11 +490,12 @@ bool UUIFrameworkPlayerComponent::TrySetFocus(float DeltaTime, FUIFrameworkWidge
 				const FUIFrameworkWidgetTreeEntry* RootEntry = GetWidgetTree().FindRootEntryById(CurrentWidgetToFocus);
 				if (RootEntry)
 				{
-					if (FUIFrameworkGameLayerSlot* Slot = RootList.FindEntry(RootEntry->ChildId))
+					if (FUIFrameworkGameLayerSlot* LayerSlot = RootList.FindEntry(RootEntry->ChildId))
 					{
-						if (Slot->InputMode == EUIFrameworkInputMode::UI)
+						if (LayerSlot->InputMode == EUIFrameworkInputMode::UI)
 						{
-							UMGWidget->SetFocus();
+							check(Presenter);
+							Presenter->FocusWidget(UMGWidget);
 						}
 					}
 				}
diff --git a/Engine/Plugins/Experimental/UIFramework/Source/Private/UIFPresenter.cpp b/Engine/Plugins/Experimental/UIFramework/Source/Private/UIFPresenter.cpp
--- a/Engine/Plugins/Experimental/UIFramework/Source/Private/UIFPresenter.cpp
+++ b/Engine/Plugins/Experimental/UIFramework/Source/Private/UIFPresenter.cpp
@@ -48,6 +48,14 @@ void UUIFrameworkGameViewportPresenter::RemoveFromViewport(FUIFrameworkWidgetId
 	}
 }
 
+void UUIFrameworkGameViewportPresenter::FocusWidget(UWidget* UMGWidget)
+{
+	if (UMGWidget)
+	{
+		UMGWidget->SetFocus();
+	}
+}
+
 void UUIFrameworkGameViewportPresenter::BeginDestroy()
 {
 	if (GetOuter() && !IsTemplate())
diff --git a/Engine/Plugins/Experimental/UIFramework/Source/Public/UIFPresenter.h b/Engine/Plugins/Experimental/UIFramework/Source/Public/UIFPresenter.h
--- a/Engine/Plugins/Experimental/UIFramework/Source/Public/UIFPresenter.h
+++ b/Engine/Plugins/Experimental/UIFramework/Source/Public/UIFPresenter.h
@@ -28,6 +28,10 @@ public:
 	virtual void RemoveFromViewport(FUIFrameworkWidgetId WidgetId)
 	{
 
+	}
+	virtual void FocusWidget(UWidget* UMGWidget)
+	{
+
 	}
 };
 
@@ -43,7 +47,7 @@ class UUIFrameworkGameViewportPresenter : public UUIFrameworkPresenter
 public:
 	UE_API virtual void AddToViewport(UWidget* UMGWidget, const FUIFrameworkGameLayerSlot& Slot) override;
 	UE_API virtual void RemoveFromViewport(FUIFrameworkWidgetId WidgetId) override;
-
+	UE_API virtual void FocusWidget(UWidget* UMGWidget) override;
 	UE_API virtual void BeginDestroy() override;
 
 private:
