Expose focus events in SButton
#okforversepublic

diff --git a/Engine/Plugins/Experimental/UIFramework/Source/Private/Widgets/UIFButton.cpp b/Engine/Plugins/Experimental/UIFramework/Source/Private/Widgets/UIFButton.cpp
--- a/Engine/Plugins/Experimental/UIFramework/Source/Private/Widgets/UIFButton.cpp
+++ b/Engine/Plugins/Experimental/UIFramework/Source/Private/Widgets/UIFButton.cpp
@@ -99,6 +99,10 @@ void UUIFrameworkButton::LocalOnUMGWidgetCreated()
 	Super::LocalOnUMGWidgetCreated();
 	UButton* Button = CastChecked<UButton>(LocalGetUMGWidget());
 	Button->OnClicked.AddUniqueDynamic(this, &ThisClass::HandleClick);
+	Button->OnReceivedFocus.BindUObject(this, &ThisClass::HandleFocusReceived);
+	Button->OnLostFocus.BindUObject(this, &ThisClass::HandleFocusLost);
+	Button->OnHovered.AddUniqueDynamic(this, &ThisClass::HandleHovered);
+	Button->OnUnhovered.AddUniqueDynamic(this, &ThisClass::HandleUnhovered);
 }
 
 
@@ -144,6 +148,58 @@ void UUIFrameworkButton::ServerClick_Implementation(APlayerController* PlayerCon
 	OnClick.Broadcast(Argument);
 }
 
+void UUIFrameworkButton::HandleFocusReceived()
+{
+	ServerFocusReceived(GetPlayerController());
+}
+
+void UUIFrameworkButton::ServerFocusReceived_Implementation(APlayerController* PlayerController)
+{
+	FUIFrameworkFocusEventArgument Argument;
+	Argument.PlayerController = PlayerController;
+	Argument.Sender = this;
+	OnFocusReceived.Broadcast(Argument);
+}
+
+void UUIFrameworkButton::HandleFocusLost()
+{
+	ServerFocusLost(GetPlayerController());
+}
+
+void UUIFrameworkButton::ServerFocusLost_Implementation(APlayerController* PlayerController)
+{
+	FUIFrameworkFocusEventArgument Argument;
+	Argument.PlayerController = PlayerController;
+	Argument.Sender = this;
+	OnFocusLost.Broadcast(Argument);
+}
+
+void UUIFrameworkButton::HandleHovered()
+{
+	ServerHovered(GetPlayerController());
+}
+
+void UUIFrameworkButton::ServerHovered_Implementation(APlayerController* PlayerController)
+{
+	FUIFrameworkHoverEventArgument Argument;
+	Argument.PlayerController = PlayerController;
+	Argument.Sender = this;
+	OnHovered.Broadcast(Argument);
+}
+
+void UUIFrameworkButton::HandleUnhovered()
+{
+	ServerUnhovered(GetPlayerController());
+}
+
+void UUIFrameworkButton::ServerUnhovered_Implementation(APlayerController* PlayerController)
+{
+	FUIFrameworkHoverEventArgument Argument;
+	Argument.PlayerController = PlayerController;
+	Argument.Sender = this;
+	OnUnhovered.Broadcast(Argument);
+}
+
 void UUIFrameworkButton::OnRep_Slot()
 {
 	if (LocalGetUMGWidget() && Slot.LocalIsAquiredWidgetValid())
diff --git a/Engine/Plugins/Experimental/UIFramework/Source/Public/Types/UIFEvents.h b/Engine/Plugins/Experimental/UIFramework/Source/Public/Types/UIFEvents.h
--- a/Engine/Plugins/Experimental/UIFramework/Source/Public/Types/UIFEvents.h
+++ b/Engine/Plugins/Experimental/UIFramework/Source/Public/Types/UIFEvents.h
@@ -40,3 +40,29 @@ struct FUIFrameworkClickEventArgument : public FUIFrameworkSimpleEventArgument
 };
 
 DECLARE_MULTICAST_DELEGATE_OneParam(FUIFrameworkClickEvent, FUIFrameworkClickEventArgument);
+
+/**
+ *
+ */
+USTRUCT(BlueprintType)
+struct FUIFrameworkFocusEventArgument : public FUIFrameworkSimpleEventArgument
+{
+	GENERATED_BODY()
+
+	FUIFrameworkFocusEventArgument() = default;
+};
+
+DECLARE_MULTICAST_DELEGATE_OneParam(FUIFrameworkFocusEvent, FUIFrameworkFocusEventArgument);
+
+/**
+ *
+ */
+USTRUCT(BlueprintType)
+struct FUIFrameworkHoverEventArgument : public FUIFrameworkSimpleEventArgument
+{
+	GENERATED_BODY()
+
+	FUIFrameworkHoverEventArgument() = default;
+};
+
+DECLARE_MULTICAST_DELEGATE_OneParam(FUIFrameworkHoverEvent, FUIFrameworkHoverEventArgument);
\ No newline at end of file
diff --git a/Engine/Plugins/Experimental/UIFramework/Source/Public/Widgets/UIFButton.h b/Engine/Plugins/Experimental/UIFramework/Source/Public/Widgets/UIFButton.h
--- a/Engine/Plugins/Experimental/UIFramework/Source/Public/Widgets/UIFButton.h
+++ b/Engine/Plugins/Experimental/UIFramework/Source/Public/Widgets/UIFButton.h
@@ -47,12 +47,43 @@ private:
 	UFUNCTION(Server, Reliable)
 	UE_API void ServerClick(APlayerController* PlayerController);
 
+	UFUNCTION()
+	UE_API void HandleFocusReceived();
+
+	UFUNCTION(Server, Reliable)
+	UE_API void ServerFocusReceived(APlayerController* PlayerController);
+
+	UFUNCTION()
+	UE_API void HandleFocusLost();
+
+	UFUNCTION(Server, Reliable)
+	UE_API void ServerFocusLost(APlayerController* PlayerController);
+
+	UFUNCTION()
+	UE_API void HandleHovered();
+
+	UFUNCTION(Server, Reliable)
+	UE_API void ServerHovered(APlayerController* PlayerController);
+
+	UFUNCTION()
+	UE_API void HandleUnhovered();
+
+	UFUNCTION(Server, Reliable)
+	UE_API void ServerUnhovered(APlayerController* PlayerController);
+
 	UFUNCTION()
 	UE_API void OnRep_Slot();
 
 public:
 	FUIFrameworkClickEvent OnClick;
 
+	FUIFrameworkFocusEvent OnFocusReceived;
+
+	FUIFrameworkFocusEvent OnFocusLost;
+
+	FUIFrameworkHoverEvent OnHovered;
+
+	FUIFrameworkHoverEvent OnUnhovered;
 private:
 	UPROPERTY(/*ExposeOnSpawn, */ReplicatedUsing = OnRep_Slot)
 	FUIFrameworkSimpleSlot Slot;
diff --git a/Engine/Plugins/Runtime/CommonUI/Source/CommonUI/Private/CommonButtonBase.cpp b/Engine/Plugins/Runtime/CommonUI/Source/CommonUI/Private/CommonButtonBase.cpp
--- a/Engine/Plugins/Runtime/CommonUI/Source/CommonUI/Private/CommonButtonBase.cpp
+++ b/Engine/Plugins/Runtime/CommonUI/Source/CommonUI/Private/CommonButtonBase.cpp
@@ -390,16 +390,6 @@ FReply UCommonButtonInternalBase::SlateHandleDropOverride(const FGeometry& MyGeo
 }
 //	End Drag Drop
 
-void UCommonButtonInternalBase::SlateHandleOnReceivedFocus()
-{
-	OnReceivedFocus.ExecuteIfBound();
-}
-
-void UCommonButtonInternalBase::SlateHandleOnLostFocus()
-{
-	OnLostFocus.ExecuteIfBound();
-}
-
 //////////////////////////////////////////////////////////////////////////
 // UCommonButtonBase
 //////////////////////////////////////////////////////////////////////////
diff --git a/Engine/Plugins/Runtime/CommonUI/Source/CommonUI/Private/CommonButtonTypes.cpp b/Engine/Plugins/Runtime/CommonUI/Source/CommonUI/Private/CommonButtonTypes.cpp
--- a/Engine/Plugins/Runtime/CommonUI/Source/CommonUI/Private/CommonButtonTypes.cpp
+++ b/Engine/Plugins/Runtime/CommonUI/Source/CommonUI/Private/CommonButtonTypes.cpp
@@ -194,22 +194,6 @@ bool SCommonButton::IsInteractable() const
 	return bIsButtonEnabled && bIsInteractionEnabled;
 }
 
-/** Overridden to fire delegate for external listener */
-FReply SCommonButton::OnFocusReceived(const FGeometry& MyGeometry, const FFocusEvent& InFocusEvent)
-{
-	FReply ReturnReply = SButton::OnFocusReceived(MyGeometry, InFocusEvent);
-	OnReceivedFocus.ExecuteIfBound();
-
-	return ReturnReply;
-}
-
-void SCommonButton::OnFocusLost(const FFocusEvent& InFocusEvent)
-{
-	SButton::OnFocusLost(InFocusEvent);
-
-	OnLostFocus.ExecuteIfBound();
-}
-
 int32 SCommonButton::OnPaint(const FPaintArgs& Args, const FGeometry& AllottedGeometry, const FSlateRect& MyClippingRect, FSlateWindowElementList& OutDrawElements, int32 LayerId, const FWidgetStyle& InWidgetStyle, bool bParentEnabled) const
 {
 	bool bEnabled = bParentEnabled && bIsButtonEnabled;
diff --git a/Engine/Plugins/Runtime/CommonUI/Source/CommonUI/Private/CommonButtonTypes.h b/Engine/Plugins/Runtime/CommonUI/Source/CommonUI/Private/CommonButtonTypes.h
--- a/Engine/Plugins/Runtime/CommonUI/Source/CommonUI/Private/CommonButtonTypes.h
+++ b/Engine/Plugins/Runtime/CommonUI/Source/CommonUI/Private/CommonButtonTypes.h
@@ -77,6 +77,8 @@ public:
 			.OnPressed(InArgs._OnPressed)
 			.OnReleased(InArgs._OnReleased)
 			.IsFocusable(InArgs._IsFocusable)
+			.OnReceivedFocus(InArgs._OnReceivedFocus)
+			.OnLostFocus(InArgs._OnLostFocus)
 			.Content()
 			[
 				InArgs._Content.Widget
@@ -86,8 +88,6 @@ public:
 		// Set the hover state to indicate that we want to override the default behavior
 		SetHover(false);
 
-		OnReceivedFocus = InArgs._OnReceivedFocus;
-		OnLostFocus = InArgs._OnLostFocus;
 		bIsButtonEnabled = InArgs._IsButtonEnabled;
 		bIsInteractionEnabled = InArgs._IsInteractionEnabled;
 		bHovered = false;
@@ -117,11 +117,6 @@ public:
 
 	UE_API bool IsInteractable() const;
 
-	/** Overridden to fire delegate for external listener */
-	UE_API virtual FReply OnFocusReceived(const FGeometry& MyGeometry, const FFocusEvent& InFocusEvent);
-
-	UE_API virtual void OnFocusLost(const FFocusEvent& InFocusEvent) override;
-
 	UE_API virtual int32 OnPaint(const FPaintArgs& Args, const FGeometry& AllottedGeometry, const FSlateRect& MyClippingRect, FSlateWindowElementList& OutDrawElements, int32 LayerId, const FWidgetStyle& InWidgetStyle, bool bParentEnabled) const override;
 
 protected:
@@ -131,12 +126,6 @@ protected:
 private:
 	FOnClicked OnDoubleClicked;
 
-	/** Delegate fired whenever focus is received */
-	FSimpleDelegate OnReceivedFocus;
-
-	/** Delegate fired whenever focus is lost */
-	FSimpleDelegate OnLostFocus;
-
 	/** True if the button is enabled */
 	bool bIsButtonEnabled;
 
diff --git a/Engine/Plugins/Runtime/CommonUI/Source/CommonUI/Public/CommonButtonBase.h b/Engine/Plugins/Runtime/CommonUI/Source/CommonUI/Public/CommonButtonBase.h
--- a/Engine/Plugins/Runtime/CommonUI/Source/CommonUI/Public/CommonButtonBase.h
+++ b/Engine/Plugins/Runtime/CommonUI/Source/CommonUI/Public/CommonButtonBase.h
@@ -267,12 +267,6 @@ public:
 	UPROPERTY(BlueprintAssignable, Category = "Common Button Internal|Event")
 	FOnButtonClickedEvent OnDoubleClicked;
 
-	/** Called when the button receives focus */
-	FSimpleDelegate OnReceivedFocus;
-
-	/** Called when the button loses focus */
-	FSimpleDelegate OnLostFocus;
-
 protected:
 	// UWidget interface
 	COMMONUI_API virtual TSharedRef<SWidget> RebuildWidget() override;
@@ -289,12 +283,6 @@ protected:
 	COMMONUI_API virtual FReply SlateHandleDragOverOverride(const FGeometry& MyGeometry, const FDragDropEvent& DragDropEvent);
 	COMMONUI_API virtual FReply SlateHandleDropOverride(const FGeometry& MyGeometry, const FDragDropEvent& DragDropEvent);
 
-	/** Called when internal slate button receives focus; Fires OnReceivedFocus */
-	COMMONUI_API void SlateHandleOnReceivedFocus();
-
-	/** Called when internal slate button loses focus; Fires OnLostFocus */
-	COMMONUI_API void SlateHandleOnLostFocus();
-
 	/** The minimum width of the button */
 	UPROPERTY()
 	int32 MinWidth;
diff --git a/Engine/Source/Runtime/Slate/Private/Widgets/Input/SButton.cpp b/Engine/Source/Runtime/Slate/Private/Widgets/Input/SButton.cpp
--- a/Engine/Source/Runtime/Slate/Private/Widgets/Input/SButton.cpp
+++ b/Engine/Source/Runtime/Slate/Private/Widgets/Input/SButton.cpp
@@ -108,6 +108,8 @@ void SButton::Construct( const FArguments& InArgs )
 	OnReleased = InArgs._OnReleased;
 	OnHovered = InArgs._OnHovered;
 	OnUnhovered = InArgs._OnUnhovered;
+	OnReceivedFocus = InArgs._OnReceivedFocus;
+	OnLostFocus = InArgs._OnLostFocus;
 
 	OnSlateButtonDragDetected = InArgs._OnSlateButtonDragDetected;
 	OnSlateButtonDragEnter = InArgs._OnSlateButtonDragEnter;
@@ -277,6 +279,15 @@ void SButton::OnFocusLost( const FFocusEvent& InFocusEvent )
 	SBorder::OnFocusLost(InFocusEvent);
 
 	Release();
+	OnLostFocus.ExecuteIfBound();
+}
+
+FReply SButton::OnFocusReceived(const FGeometry& MyGeometry, const FFocusEvent& InFocusEvent)
+{
+	FReply ReturnReply = SBorder::OnFocusReceived(MyGeometry, InFocusEvent);
+	OnReceivedFocus.ExecuteIfBound();
+
+	return ReturnReply;
 }
 
 FReply SButton::OnKeyDown( const FGeometry& MyGeometry, const FKeyEvent& InKeyEvent )
@@ -693,6 +704,16 @@ void SButton::SetOnUnhovered(FSimpleDelegate InOnUnhovered)
 	OnUnhovered = InOnUnhovered;
 }
 
+void SButton::SetOnFocusReceived(FSimpleDelegate InOnFocusReceived)
+{
+	OnReceivedFocus = InOnFocusReceived;
+}
+
+void SButton::SetOnFocusLost(FSimpleDelegate InOnFocusLost)
+{
+	OnLostFocus = InOnFocusLost;
+}
+
 void SButton::ExecuteHoverStateChanged(bool bPlaySound)
 {
 	if (IsHovered())
diff --git a/Engine/Source/Runtime/Slate/Public/Widgets/Input/SButton.h b/Engine/Source/Runtime/Slate/Public/Widgets/Input/SButton.h
--- a/Engine/Source/Runtime/Slate/Public/Widgets/Input/SButton.h
+++ b/Engine/Source/Runtime/Slate/Public/Widgets/Input/SButton.h
@@ -98,6 +98,10 @@ public:
 
 		SLATE_EVENT( FSimpleDelegate, OnUnhovered )
 
+		SLATE_EVENT( FSimpleDelegate, OnReceivedFocus )
+
+		SLATE_EVENT( FSimpleDelegate, OnLostFocus )
+
 		// Drag and Drop
 		SLATE_EVENT(FOnDragDetected, OnSlateButtonDragDetected)
 
@@ -204,6 +208,12 @@ public:
 	/** Set OnUnhovered event */
 	SLATE_API void SetOnUnhovered(FSimpleDelegate InOnUnhovered);
 
+	/** Set OnFocusReceived event */
+	SLATE_API void SetOnFocusReceived(FSimpleDelegate InOnFocusReceived);
+
+	/** Set OnFocusLost event */
+	SLATE_API void SetOnFocusLost(FSimpleDelegate InOnFocusLost);
+
 	/** See ButtonStyle attribute */
 	SLATE_API void SetButtonStyle(const FButtonStyle* ButtonStyle);
 
@@ -224,6 +234,7 @@ public:
 	SLATE_API virtual int32 OnPaint(const FPaintArgs& Args, const FGeometry& AllottedGeometry, const FSlateRect& MyCullingRect, FSlateWindowElementList& OutDrawElements, int32 LayerId, const FWidgetStyle& InWidgetStyle, bool bParentEnabled) const override;
 	SLATE_API virtual bool SupportsKeyboardFocus() const override;
 	SLATE_API virtual void OnFocusLost( const FFocusEvent& InFocusEvent ) override;
+	SLATE_API virtual FReply OnFocusReceived(const FGeometry& MyGeometry, const FFocusEvent& InFocusEvent) override;
 	SLATE_API virtual FReply OnKeyDown( const FGeometry& MyGeometry, const FKeyEvent& InKeyEvent ) override;
 	SLATE_API virtual FReply OnKeyUp( const FGeometry& MyGeometry, const FKeyEvent& InKeyEvent ) override;
 	SLATE_API virtual FReply OnMouseButtonDown( const FGeometry& MyGeometry, const FPointerEvent& MouseEvent ) override;
@@ -332,6 +343,12 @@ private:
 	/** The delegate to execute when the button exit the hovered state */
 	FSimpleDelegate OnUnhovered;
 
+	/** Delegate fired whenever focus is received */
+	FSimpleDelegate OnReceivedFocus;
+
+	/** Delegate fired whenever focus is lost */
+	FSimpleDelegate OnLostFocus;
+
 	/** Delegate triggered when a user starts to drag a button */
 	FOnDragDetected OnSlateButtonDragDetected;
 
diff --git a/Engine/Source/Runtime/UMG/Private/Components/Button.cpp b/Engine/Source/Runtime/UMG/Private/Components/Button.cpp
--- a/Engine/Source/Runtime/UMG/Private/Components/Button.cpp
+++ b/Engine/Source/Runtime/UMG/Private/Components/Button.cpp
@@ -70,6 +70,8 @@ TSharedRef<SWidget> UButton::RebuildWidget()
 		.OnReleased(BIND_UOBJECT_DELEGATE(FSimpleDelegate, SlateHandleReleased))
 		.OnHovered_UObject(this, &ThisClass::SlateHandleHovered)
 		.OnUnhovered_UObject(this, &ThisClass::SlateHandleUnhovered)
+		.OnReceivedFocus_UObject(this, &ThisClass::SlateHandleOnReceivedFocus)
+		.OnLostFocus_UObject(this, &ThisClass::SlateHandleOnLostFocus)
 		.OnSlateButtonDragDetected(BIND_UOBJECT_DELEGATE(FOnDragDetected, SlateHandleDragDetected))
 		.OnSlateButtonDragEnter(BIND_UOBJECT_DELEGATE(FOnDragEnter, SlateHandleDragEnter))
 		.OnSlateButtonDragLeave(BIND_UOBJECT_DELEGATE(FOnDragLeave, SlateHandleDragLeave))
@@ -304,6 +306,16 @@ void UButton::SlateHandleUnhovered()
 	BroadcastBinaryPostStateChange(UWidgetHoveredStateRegistration::Bit, false);
 }
 
+void UButton::SlateHandleOnReceivedFocus()
+{
+	OnReceivedFocus.ExecuteIfBound();
+}
+
+void UButton::SlateHandleOnLostFocus()
+{
+	OnLostFocus.ExecuteIfBound();
+}
+
 FReply UButton::SlateHandleDragDetected(const FGeometry& MyGeometry, const FPointerEvent& MouseEvent)
 {
 	if (OnButtonDragDetected.IsBound())
diff --git a/Engine/Source/Runtime/UMG/Public/Components/Button.h b/Engine/Source/Runtime/UMG/Public/Components/Button.h
--- a/Engine/Source/Runtime/UMG/Public/Components/Button.h
+++ b/Engine/Source/Runtime/UMG/Public/Components/Button.h
@@ -104,6 +104,12 @@ public:
 	/** Delegate triggered when a user's drag is dropped in the bounds of this button */
 	FOnDrop OnButtonDrop;
 
+	/** Called when the button receives focus */
+	FSimpleDelegate OnReceivedFocus;
+
+	/** Called when the button loses focus */
+	FSimpleDelegate OnLostFocus;
+
 public:
 	
 	/** Sets the color multiplier for the button background */
@@ -185,6 +191,8 @@ protected:
 	UMG_API void SlateHandleReleased();
 	UMG_API void SlateHandleHovered();
 	UMG_API void SlateHandleUnhovered();
+	UMG_API void SlateHandleOnReceivedFocus();
+	UMG_API void SlateHandleOnLostFocus();
 
 	// Drag and Drop
 	UMG_API FReply SlateHandleDragDetected(const FGeometry& MyGeometry, const FPointerEvent& MouseEvent);
