VerseUI: Better logging to identify material misuse in image blocks.

#rnx

diff --git a/Engine/Plugins/Experimental/UIFramework/Source/Private/Widgets/UIFImageBlock.cpp b/Engine/Plugins/Experimental/UIFramework/Source/Private/Widgets/UIFImageBlock.cpp
--- a/Engine/Plugins/Experimental/UIFramework/Source/Private/Widgets/UIFImageBlock.cpp
+++ b/Engine/Plugins/Experimental/UIFramework/Source/Private/Widgets/UIFImageBlock.cpp
@@ -7,6 +7,7 @@
 #include "Engine/Texture2D.h"
 #include "Materials/Material.h"
 #include "UIFLocalSettings.h"
+#include "UIFLog.h"
 
 #include "Net/Core/PushModel/PushModel.h"
 #include "Net/UnrealNetwork.h"
@@ -170,6 +171,10 @@ UObject* UUIFrameworkImageBlock::ProcessResourceObject()
 				{
 					Result = ResourceObject;
 				}
+				else
+				{
+					UE_LOG(LogUIFramework, Error, TEXT("Can't use material for material_block %s. The material domain is not set to \"User Interface\""), *GetName(), *BaseMaterial->GetName());
+				}
 			}
 		}
 		else if (UTexture2D* Texture = Cast<UTexture2D>(ResourceObject))
