VerseUI: Fix log formatting

#rnx

diff --git a/Engine/Plugins/Experimental/UIFramework/Source/Private/Widgets/UIFImageBlock.cpp b/Engine/Plugins/Experimental/UIFramework/Source/Private/Widgets/UIFImageBlock.cpp
--- a/Engine/Plugins/Experimental/UIFramework/Source/Private/Widgets/UIFImageBlock.cpp
+++ b/Engine/Plugins/Experimental/UIFramework/Source/Private/Widgets/UIFImageBlock.cpp
@@ -173,7 +173,7 @@ UObject* UUIFrameworkImageBlock::ProcessResourceObject()
 				}
 				else
 				{
-					UE_LOG(LogUIFramework, Error, TEXT("Can't use material for material_block %s. The material domain is not set to \"User Interface\""), *GetName(), *BaseMaterial->GetName());
+					UE_LOG(LogUIFramework, Error, TEXT("Can't use material %s for material_block %s. The material domain is not set to \"User Interface\""), *BaseMaterial->GetName(), *GetName());
 				}
 			}
 		}
