UIFramework: Create the Presenter on-demand as it can fail during construction

#rnx

diff --git a/Engine/Plugins/Experimental/UIFramework/Source/Private/UIFPlayerComponent.cpp b/Engine/Plugins/Experimental/UIFramework/Source/Private/UIFPlayerComponent.cpp
--- a/Engine/Plugins/Experimental/UIFramework/Source/Private/UIFPlayerComponent.cpp
+++ b/Engine/Plugins/Experimental/UIFramework/Source/Private/UIFPlayerComponent.cpp
@@ -3,7 +3,6 @@
 #include "UIFPlayerComponent.h"
 #include "Misc/MemStack.h"
 #include "UIFLog.h"
-#include "Templates/NonNullPointer.h"
 #include "UIFModule.h"
 #include "Types/UIFWidgetOwner.h"
 #include "UIFLocalSettings.h"
@@ -29,9 +28,9 @@ void FUIFrameworkGameLayerSlotList::PreReplicatedRemove(const TArrayView<int32>&
 		FUIFrameworkGameLayerSlot& Slot = Entries[Index];
 		if (Slot.LocalIsAquiredWidgetValid())
 		{
-			if (Owner->Presenter)
+			if (UUIFrameworkPresenter* Presenter = Owner->GetOrCreatePresenter())
 			{
-				Owner->Presenter->RemoveFromViewport(Slot.GetWidgetId());
+				Presenter->RemoveFromViewport(Slot.GetWidgetId());
 			}
 		}
 	}
@@ -92,11 +91,6 @@ UUIFrameworkPlayerComponent::UUIFrameworkPlayerComponent()
 	PrimaryComponentTick.bCanEverTick = true;
 	PrimaryComponentTick.bStartWithTickEnabled = false;
 	PrimaryComponentTick.TickGroup = ETickingGroup::TG_DuringPhysics;
-
-	if (!IsTemplate() && GetPlayerController() && GetPlayerController()->IsLocalPlayerController())
-	{
-		Presenter = static_cast<UUIFrameworkPresenter*>(CreateDefaultSubobject(TEXT("Presenter"), UUIFrameworkPresenter::StaticClass(), FUIFrameworkModule::GetPresenterClass().Get(), true, true));
-	}
 }
 
 void UUIFrameworkPlayerComponent::InitializeComponent()
@@ -273,6 +267,17 @@ void UUIFrameworkPlayerComponent::HandleAddPending()
 
 }
 
+UUIFrameworkPresenter* UUIFrameworkPlayerComponent::GetOrCreatePresenter()
+{
+	if (Presenter == nullptr)
+	{
+		Presenter = NewObject<UUIFrameworkPresenter>(this, FUIFrameworkModule::GetPresenterClass());
+	}
+
+	check(Presenter);
+	return Presenter;
+}
+
 void UUIFrameworkPlayerComponent::TickComponent(float DeltaTime, enum ELevelTick TickType, FActorComponentTickFunction* ThisTickFunction)
 {
 	APlayerController* LocalOwner = GetPlayerController();
@@ -460,8 +465,10 @@ void UUIFrameworkPlayerComponent::LocalAddChild(FUIFrameworkWidgetId WidgetId)
 			{
 				UMGWidget->RemoveFromParent();
 				LayerSlot->LocalAquireWidget();
-				check(Presenter);
-				Presenter->AddToViewport(UMGWidget, *LayerSlot);
+				if (UUIFrameworkPresenter* PresenterPtr = GetOrCreatePresenter())
+				{
+					PresenterPtr->AddToViewport(UMGWidget, *LayerSlot);
+				}
 			}
 		}
 		else
@@ -494,8 +501,10 @@ bool UUIFrameworkPlayerComponent::TrySetFocus(float DeltaTime, FUIFrameworkWidge
 					{
 						if (LayerSlot->InputMode == EUIFrameworkInputMode::UI)
 						{
-							check(Presenter);
-							Presenter->FocusWidget(UMGWidget);
+							if (UUIFrameworkPresenter* PresenterPtr = GetOrCreatePresenter())
+							{
+								PresenterPtr->FocusWidget(UMGWidget);
+							}
 						}
 					}
 				}
diff --git a/Engine/Plugins/Experimental/UIFramework/Source/Public/UIFPlayerComponent.h b/Engine/Plugins/Experimental/UIFramework/Source/Public/UIFPlayerComponent.h
--- a/Engine/Plugins/Experimental/UIFramework/Source/Public/UIFPlayerComponent.h
+++ b/Engine/Plugins/Experimental/UIFramework/Source/Public/UIFPlayerComponent.h
@@ -193,6 +193,9 @@ private:
 	UE_API bool TrySetFocus(float DeltaTime, FUIFrameworkWidgetId CurrentWidgetToFocus);
 
 	void HandleAddPending();
+
+	UUIFrameworkPresenter* GetOrCreatePresenter();
+	
 private:
 	UPROPERTY(Replicated)
 	FUIFrameworkGameLayerSlotList RootList;
