Verse UI button: fire highlight event on long press and unhighlight event on touch end

diff --git a/Engine/Plugins/Experimental/UIFramework/Source/Private/Widgets/SUIFrameworkButtonWidget.cpp b/Engine/Plugins/Experimental/UIFramework/Source/Private/Widgets/SUIFrameworkButtonWidget.cpp
new file mode 100644
--- /dev/null
+++ b/Engine/Plugins/Experimental/UIFramework/Source/Private/Widgets/SUIFrameworkButtonWidget.cpp
@@ -0,0 +1,66 @@
+// Copyright Epic Games, Inc. All Rights Reserved.
+
+#include "SUIFrameworkButtonWidget.h"
+
+SLATE_IMPLEMENT_WIDGET(SUIFrameworkButtonWidget)
+
+void SUIFrameworkButtonWidget::Construct(const FArguments& InArgs)
+{
+	OnTouchLongPressedDelegate = InArgs._OnTouchLongPressedDelegate;
+	OnTouchStartedDelegate = InArgs._OnTouchStartedDelegate;
+	OnTouchEndedDelegate = InArgs._OnTouchEndedDelegate;
+
+	SButton::Construct(SButton::FArguments()
+		.ButtonStyle(InArgs._ButtonStyle)
+		.OnSlateButtonDragDetected(InArgs._OnSlateButtonDragDetected)
+		.OnSlateButtonDragOver(InArgs._OnSlateButtonDragOver)
+		.OnSlateButtonDragEnter(InArgs._OnSlateButtonDragEnter)
+		.OnSlateButtonDragLeave(InArgs._OnSlateButtonDragLeave)
+		.OnSlateButtonDrop(InArgs._OnSlateButtonDrop)
+		.ClickMethod(InArgs._ClickMethod)
+		.TouchMethod(InArgs._TouchMethod)
+		.PressMethod(InArgs._PressMethod)
+		.OnClicked(InArgs._OnClicked)
+		.OnClicked(InArgs._OnClicked)
+		.OnPressed(InArgs._OnPressed)
+		.OnReleased(InArgs._OnReleased)
+		.IsFocusable(InArgs._IsFocusable)
+		.OnReceivedFocus(InArgs._OnReceivedFocus)
+		.OnLostFocus(InArgs._OnLostFocus)
+		.OnHovered(InArgs._OnHovered)
+		.OnUnhovered(InArgs._OnUnhovered)
+		.AllowDragDrop(InArgs._AllowDragDrop));
+}
+
+FReply SUIFrameworkButtonWidget::OnTouchGesture(const FGeometry& MyGeometry, const FPointerEvent& GestureEvent)
+{
+	if (GestureEvent.GetGestureType() == EGestureEvent::LongPress)
+	{
+		if (OnTouchLongPressedDelegate.IsBound())
+		{
+			return OnTouchLongPressedDelegate.Execute(GestureEvent);
+		}
+	}
+	return FReply::Unhandled();
+}
+
+FReply SUIFrameworkButtonWidget::OnPreviewMouseButtonDown(const FGeometry& MyGeometry, const FPointerEvent& MouseEvent)
+{
+	if (MouseEvent.IsTouchEvent())
+	{
+		if (OnTouchStartedDelegate.IsBound())
+		{
+			return OnTouchStartedDelegate.Execute(MouseEvent);
+		}
+	}
+	return FReply::Unhandled();
+}
+
+FReply SUIFrameworkButtonWidget::OnTouchEnded(const FGeometry& MyGeometry, const FPointerEvent& InTouchEvent)
+{
+	if (OnTouchEndedDelegate.IsBound())
+	{
+		return OnTouchEndedDelegate.Execute(InTouchEvent);
+	}
+	return FReply::Unhandled();
+}
\ No newline at end of file
diff --git a/Engine/Plugins/Experimental/UIFramework/Source/Private/Widgets/SUIFrameworkButtonWidget.h b/Engine/Plugins/Experimental/UIFramework/Source/Private/Widgets/SUIFrameworkButtonWidget.h
new file mode 100644
--- /dev/null
+++ b/Engine/Plugins/Experimental/UIFramework/Source/Private/Widgets/SUIFrameworkButtonWidget.h
@@ -0,0 +1,94 @@
+// Copyright Epic Games, Inc. All Rights Reserved.
+
+#pragma once
+
+#include "CoreMinimal.h"
+#include "Input/Reply.h"
+#include "Widgets/DeclarativeSyntaxSupport.h"
+#include "Widgets/Input/SButton.h"
+
+class FUIFInputPreProcessor;
+
+DECLARE_DELEGATE_RetVal_OneParam(
+	FReply,
+	FOnTouch,
+	const FPointerEvent&)
+
+/**
+ * Slate button for UUIFrameworkButtonWidget. It exposes touch events for start, end, and long press.
+ */
+class SUIFrameworkButtonWidget : public SButton
+{
+
+public:
+	SLATE_BEGIN_ARGS(SUIFrameworkButtonWidget)
+		: _OnTouchLongPressedDelegate()
+		, _OnTouchStartedDelegate()
+		, _OnTouchEndedDelegate()
+		, _ButtonStyle(&FCoreStyle::Get().GetWidgetStyle< FButtonStyle >("Button"))
+		, _OnSlateButtonDragDetected()
+		, _OnSlateButtonDragEnter()
+		, _OnSlateButtonDragLeave()
+		, _OnSlateButtonDragOver()
+		, _OnSlateButtonDrop()
+		, _ClickMethod(EButtonClickMethod::DownAndUp)
+		, _TouchMethod(EButtonTouchMethod::DownAndUp)
+		, _PressMethod(EButtonPressMethod::DownAndUp)
+		, _IsFocusable(true)
+		, _AllowDragDrop(false)
+		{
+		}
+
+		SLATE_EVENT(FOnTouch, OnTouchLongPressedDelegate)
+		SLATE_EVENT(FOnTouch, OnTouchStartedDelegate)
+		SLATE_EVENT(FOnTouch, OnTouchEndedDelegate)
+
+		/** Slot for this button's content (optional) */
+		SLATE_DEFAULT_SLOT(FArguments, Content)
+
+		/** The visual style of the button */
+		SLATE_STYLE_ARGUMENT(FButtonStyle, ButtonStyle)
+
+		SLATE_EVENT(FOnClicked, OnClicked)
+
+		SLATE_EVENT(FSimpleDelegate, OnPressed)
+
+		SLATE_EVENT(FSimpleDelegate, OnReleased)
+
+		SLATE_EVENT(FSimpleDelegate, OnHovered)
+
+		SLATE_EVENT(FSimpleDelegate, OnUnhovered)
+
+		SLATE_EVENT(FSimpleDelegate, OnReceivedFocus)
+
+		SLATE_EVENT(FSimpleDelegate, OnLostFocus)
+
+		// Drag and Drop
+		SLATE_EVENT(FOnDragDetected, OnSlateButtonDragDetected)
+		SLATE_EVENT(FOnDragEnter, OnSlateButtonDragEnter)
+		SLATE_EVENT(FOnDragLeave, OnSlateButtonDragLeave)
+		SLATE_EVENT(FOnDragOver, OnSlateButtonDragOver)
+		SLATE_EVENT(FOnDrop, OnSlateButtonDrop)
+
+		SLATE_ARGUMENT(EButtonClickMethod::Type, ClickMethod)
+		SLATE_ARGUMENT(EButtonTouchMethod::Type, TouchMethod)
+		SLATE_ARGUMENT(EButtonPressMethod::Type, PressMethod)
+
+		SLATE_ARGUMENT(bool, IsFocusable)
+		SLATE_ARGUMENT(bool, AllowDragDrop)
+	SLATE_END_ARGS()
+
+	void Construct(const FArguments& InArgs);
+
+	// SWidget overrides
+	virtual FReply OnTouchGesture(const FGeometry& MyGeometry, const FPointerEvent& GestureEvent) override;
+	virtual FReply OnPreviewMouseButtonDown(const FGeometry& MyGeometry, const FPointerEvent& MouseEvent) override;
+	virtual FReply OnTouchEnded(const FGeometry& MyGeometry, const FPointerEvent& InTouchEvent) override;
+	//~ SWidget overrides
+
+private:
+	FOnTouch OnTouchLongPressedDelegate;
+	FOnTouch OnTouchStartedDelegate;
+	FOnTouch OnTouchEndedDelegate;
+};
+
diff --git a/Engine/Plugins/Experimental/UIFramework/Source/Private/Widgets/UIFButton.cpp b/Engine/Plugins/Experimental/UIFramework/Source/Private/Widgets/UIFButton.cpp
--- a/Engine/Plugins/Experimental/UIFramework/Source/Private/Widgets/UIFButton.cpp
+++ b/Engine/Plugins/Experimental/UIFramework/Source/Private/Widgets/UIFButton.cpp
@@ -4,29 +4,16 @@
 #include "Types/UIFWidgetTree.h"
 #include "UIFLog.h"
 #include "UIFModule.h"
+#include "UIFrameworkButtonWidget.h"
 
+#include "CommonInputSubsystem.h"
 #include "Components/Button.h"
 #include "Components/ButtonSlot.h"
-
 #include "Misc/RedirectCollector.h"
 #include "Net/UnrealNetwork.h"
 
 #include UE_INLINE_GENERATED_CPP_BY_NAME(UIFButton)
 
-
-UUIFrameworkButtonWidget::UUIFrameworkButtonWidget()
-{
-	FButtonStyle TempStyle = GetStyle();
-	TempStyle.NormalPadding = FMargin(0.0f);
-	TempStyle.Normal.DrawAs = ESlateBrushDrawType::NoDrawType;
-	TempStyle.Hovered.DrawAs = ESlateBrushDrawType::NoDrawType;
-	TempStyle.Pressed.DrawAs = ESlateBrushDrawType::NoDrawType;
-	SetStyle(TempStyle);
-}
-
-/**
- * 
- */
 UUIFrameworkButton::UUIFrameworkButton()
 {
 // TODO: FIXME: Loading this asset fails. Temporary fixed by using UUIFrameworkButtonWidget::StaticClass()
@@ -97,15 +84,14 @@ void UUIFrameworkButton::AuthorityRemoveChild(UUIFrameworkWidget* Widget)
 void UUIFrameworkButton::LocalOnUMGWidgetCreated()
 {
 	Super::LocalOnUMGWidgetCreated();
-	UButton* Button = CastChecked<UButton>(LocalGetUMGWidget());
+	UUIFrameworkButtonWidget* Button = CastChecked<UUIFrameworkButtonWidget>(LocalGetUMGWidget());
 	Button->OnClicked.AddUniqueDynamic(this, &ThisClass::HandleClick);
-	Button->OnReceivedFocus.BindUObject(this, &ThisClass::HandleFocusReceived);
-	Button->OnLostFocus.BindUObject(this, &ThisClass::HandleFocusLost);
-	Button->OnHovered.AddUniqueDynamic(this, &ThisClass::HandleHovered);
-	Button->OnUnhovered.AddUniqueDynamic(this, &ThisClass::HandleUnhovered);
+	Button->OnHovered.AddUniqueDynamic(this, &ThisClass::HandleHover);
+	Button->OnUnhovered.AddUniqueDynamic(this, &ThisClass::HandleUnhover);
+	Button->OnTouchLongPress.BindUObject(this, &ThisClass::HandleHighlight);
+	Button->OnTouchEnded.BindUObject(this, &ThisClass::HandleUnhighlight);
 }
 
-
 void UUIFrameworkButton::LocalAddChild(FUIFrameworkWidgetId ChildId)
 {
 	FUIFrameworkWidgetTree* WidgetTree = GetWidgetTree();
@@ -148,56 +134,70 @@ void UUIFrameworkButton::ServerClick_Implementation(APlayerController* PlayerCon
 	OnClick.Broadcast(Argument);
 }
 
-void UUIFrameworkButton::HandleFocusReceived()
+void UUIFrameworkButton::HandleHover()
 {
-	ServerFocusReceived(GetPlayerController());
-}
+	if (const APlayerController* PlayerController = GetPlayerController())
+	{
+		if (const UCommonInputSubsystem* CommonInputSubsystem = UCommonInputSubsystem::Get(PlayerController->GetLocalPlayer()))
+		{
+			if (CommonInputSubsystem->GetCurrentInputType() == ECommonInputType::Touch)
+			{
+				return;
+			}
+		}
+	}
 
-void UUIFrameworkButton::ServerFocusReceived_Implementation(APlayerController* PlayerController)
-{
-	FUIFrameworkFocusEventArgument Argument;
-	Argument.PlayerController = PlayerController;
-	Argument.Sender = this;
-	OnFocusReceived.Broadcast(Argument);
+	HandleHighlight();
 }
 
-void UUIFrameworkButton::HandleFocusLost()
+void UUIFrameworkButton::HandleUnhover()
 {
-	ServerFocusLost(GetPlayerController());
+	if (const APlayerController* PlayerController = GetPlayerController())
+	{
+		if (const UCommonInputSubsystem* CommonInputSubsystem = UCommonInputSubsystem::Get(PlayerController->GetLocalPlayer()))
+		{
+			if (CommonInputSubsystem->GetCurrentInputType() == ECommonInputType::Touch)
+			{
+				return;
+			}
+		}
+	}
+
+	HandleUnhighlight();
 }
 
-void UUIFrameworkButton::ServerFocusLost_Implementation(APlayerController* PlayerController)
+void UUIFrameworkButton::HandleHighlight()
 {
-	FUIFrameworkFocusEventArgument Argument;
-	Argument.PlayerController = PlayerController;
-	Argument.Sender = this;
-	OnFocusLost.Broadcast(Argument);
+	if (!bIsHighlighted)
+	{
+		ServerHighlight(GetPlayerController());
+		bIsHighlighted = true;
+	}
 }
 
-void UUIFrameworkButton::HandleHovered()
+void UUIFrameworkButton::HandleUnhighlight()
 {
-	ServerHovered(GetPlayerController());
+	if (bIsHighlighted)
+	{
+		ServerUnhighlight(GetPlayerController());
+		bIsHighlighted = false;
+	}
 }
 
-void UUIFrameworkButton::ServerHovered_Implementation(APlayerController* PlayerController)
+void UUIFrameworkButton::ServerHighlight_Implementation(APlayerController* PlayerController)
 {
-	FUIFrameworkHoverEventArgument Argument;
+	FUIFrameworkHighlightEventArgument Argument;
 	Argument.PlayerController = PlayerController;
 	Argument.Sender = this;
-	OnHovered.Broadcast(Argument);
-}
-
-void UUIFrameworkButton::HandleUnhovered()
-{
-	ServerUnhovered(GetPlayerController());
+	OnHighlight.Broadcast(Argument);
 }
 
-void UUIFrameworkButton::ServerUnhovered_Implementation(APlayerController* PlayerController)
+void UUIFrameworkButton::ServerUnhighlight_Implementation(APlayerController* PlayerController)
 {
-	FUIFrameworkHoverEventArgument Argument;
+	FUIFrameworkUnhighlightEventArgument Argument;
 	Argument.PlayerController = PlayerController;
 	Argument.Sender = this;
-	OnUnhovered.Broadcast(Argument);
+	OnUnhighlight.Broadcast(Argument);
 }
 
 void UUIFrameworkButton::OnRep_Slot()
diff --git a/Engine/Plugins/Experimental/UIFramework/Source/Private/Widgets/UIFrameworkButtonWidget.cpp b/Engine/Plugins/Experimental/UIFramework/Source/Private/Widgets/UIFrameworkButtonWidget.cpp
new file mode 100644
--- /dev/null
+++ b/Engine/Plugins/Experimental/UIFramework/Source/Private/Widgets/UIFrameworkButtonWidget.cpp
@@ -0,0 +1,151 @@
+// Copyright Epic Games, Inc. All Rights Reserved.
+
+#include "UIFrameworkButtonWidget.h"
+
+#include "SUIFrameworkButtonWidget.h"
+
+#include "Components/Button.h"
+#include "Components/ButtonSlot.h"
+#include "Framework/Application/SlateApplication.h"
+#include "Widgets/DeclarativeSyntaxSupport.h"
+
+#include UE_INLINE_GENERATED_CPP_BY_NAME(UIFrameworkButtonWidget)
+
+UUIFrameworkButtonWidget::UUIFrameworkButtonWidget()
+{
+	FButtonStyle TempStyle = GetStyle();
+	TempStyle.NormalPadding = FMargin(0.0f);
+	TempStyle.Normal.DrawAs = ESlateBrushDrawType::NoDrawType;
+	TempStyle.Hovered.DrawAs = ESlateBrushDrawType::NoDrawType;
+	TempStyle.Pressed.DrawAs = ESlateBrushDrawType::NoDrawType;
+	SetStyle(TempStyle);
+}
+
+TSharedRef<SWidget> UUIFrameworkButtonWidget::RebuildWidget()
+{	
+	MyButton = UIFButton = SNew(SUIFrameworkButtonWidget)
+		.OnTouchLongPressedDelegate(BIND_UOBJECT_DELEGATE(FOnTouch, SlateHandleTouchLongPress))
+		.OnTouchStartedDelegate(BIND_UOBJECT_DELEGATE(FOnTouch, SlateHandleTouchStarted))
+		.OnTouchEndedDelegate(BIND_UOBJECT_DELEGATE(FOnTouch, SlateHandleTouchEnded))
+		.OnClicked(BIND_UOBJECT_DELEGATE(FOnClicked, SlateHandleClicked))
+		.OnPressed(BIND_UOBJECT_DELEGATE(FSimpleDelegate, SlateHandlePressed))
+		.OnReleased(BIND_UOBJECT_DELEGATE(FSimpleDelegate, SlateHandleReleased))
+		.OnHovered_UObject(this, &ThisClass::SlateHandleHovered)
+		.OnUnhovered_UObject(this, &ThisClass::SlateHandleUnhovered)
+		.OnReceivedFocus_UObject(this, &ThisClass::SlateHandleOnReceivedFocus)
+		.OnLostFocus_UObject(this, &ThisClass::SlateHandleOnLostFocus)
+		.OnSlateButtonDragDetected(BIND_UOBJECT_DELEGATE(FOnDragDetected, SlateHandleDragDetected))
+		.OnSlateButtonDragEnter(BIND_UOBJECT_DELEGATE(FOnDragEnter, SlateHandleDragEnter))
+		.OnSlateButtonDragLeave(BIND_UOBJECT_DELEGATE(FOnDragLeave, SlateHandleDragLeave))
+		.OnSlateButtonDragOver(BIND_UOBJECT_DELEGATE(FOnDragOver, SlateHandleDragOver))
+		.OnSlateButtonDrop(BIND_UOBJECT_DELEGATE(FOnDrop, SlateHandleDrop))
+		.ButtonStyle(&GetStyle())
+		.ClickMethod(GetClickMethod())
+		.TouchMethod(GetTouchMethod())
+		.PressMethod(GetPressMethod())
+		.IsFocusable(GetIsFocusable())
+		.AllowDragDrop(bAllowDragDrop);
+
+	if (GetChildrenCount() > 0)
+	{
+		Cast<UButtonSlot>(GetContentSlot())->BuildSlot(MyButton.ToSharedRef());
+	}
+
+	if (!InputPreProcessor)
+	{
+		InputPreProcessor = MakeShared<FUIFInputPreProcessor>();
+		InputPreProcessor->OnTrackedTouchEnded().BindUObject(this, &UUIFrameworkButtonWidget::HandleTouchEnded);
+	}
+
+	return UIFButton.ToSharedRef();
+}
+
+void UUIFrameworkButtonWidget::ReleaseSlateResources(bool bReleaseChildren)
+{
+	Super::ReleaseSlateResources(bReleaseChildren);
+	UIFButton.Reset();
+
+	if (InputPreProcessor)
+	{
+		InputPreProcessor->OnTrackedTouchEnded().Unbind();
+		InputPreProcessor->StopTrackingTouch();
+	}
+
+	InputPreProcessor.Reset();
+}
+
+FReply UUIFrameworkButtonWidget::SlateHandleTouchLongPress(const FPointerEvent& PointerEvent)
+{
+	if (UIFButton)
+	{
+		bLongPressActive = true;
+		OnTouchLongPress.ExecuteIfBound();
+		return FReply::Handled();
+	}
+
+	return FReply::Unhandled();
+}
+
+FReply UUIFrameworkButtonWidget::SlateHandleTouchStarted(const FPointerEvent& PointerEvent)
+{
+	bTouchStarted = true;
+	StartTrackingTouch(PointerEvent.GetPointerIndex());
+	OnTouchStarted.ExecuteIfBound();
+
+	return FReply::Unhandled();	// let the button content process the mouse down event too.
+}
+
+FReply UUIFrameworkButtonWidget::SlateHandleTouchEnded(const FPointerEvent& PointerEvent)
+{
+	return HandleTouchEnded() ? FReply::Handled() : FReply::Unhandled();
+}
+
+bool UUIFrameworkButtonWidget::HandleTouchEnded()
+{
+	const bool bShouldPreventClick = bLongPressActive;
+	bTouchStarted = false;
+	StopTrackingTouch();
+	if (bLongPressActive)
+	{
+		OnTouchEnded.ExecuteIfBound();
+		bLongPressActive = false;
+	}
+	return bShouldPreventClick;
+}
+
+void UUIFrameworkButtonWidget::StartTrackingTouch(int32 PointerIndex)
+{
+	if (InputPreProcessor)
+	{
+		InputPreProcessor->StartTrackingTouch(PointerIndex);
+	}
+}
+
+void UUIFrameworkButtonWidget::StopTrackingTouch()
+{
+	if (InputPreProcessor)
+	{
+		InputPreProcessor->StopTrackingTouch();
+	}
+}
+
+bool FUIFInputPreProcessor::HandleMouseButtonUpEvent(FSlateApplication& SlateApp, const FPointerEvent& MouseEvent)
+{
+	if (MouseEvent.GetPointerIndex() == TrackedTouchIndex && MouseEvent.IsTouchEvent() && OnTrackedTouchEndedDelegate.IsBound())
+	{
+		return OnTrackedTouchEndedDelegate.Execute();
+	}
+	return false;
+}
+
+void FUIFInputPreProcessor::StartTrackingTouch(int32 Index)
+{
+	TrackedTouchIndex = Index;
+	FSlateApplication::Get().RegisterInputPreProcessor(AsShared());
+}
+
+void FUIFInputPreProcessor::StopTrackingTouch()
+{
+	FSlateApplication::Get().UnregisterInputPreProcessor(AsShared());
+	TrackedTouchIndex = INDEX_NONE;
+}
\ No newline at end of file
diff --git a/Engine/Plugins/Experimental/UIFramework/Source/Private/Widgets/UIFrameworkButtonWidget.h b/Engine/Plugins/Experimental/UIFramework/Source/Private/Widgets/UIFrameworkButtonWidget.h
new file mode 100644
--- /dev/null
+++ b/Engine/Plugins/Experimental/UIFramework/Source/Private/Widgets/UIFrameworkButtonWidget.h
@@ -0,0 +1,75 @@
+// Copyright Epic Games, Inc. All Rights Reserved.
+
+#pragma once
+
+#include "Components/Button.h"
+#include "Delegates/Delegate.h"
+#include "Framework/Application/IInputProcessor.h"
+#include "UIFrameworkButtonWidget.generated.h"
+
+#define UE_API UIFRAMEWORK_API
+
+class FUIFInputPreProcessor;
+
+DECLARE_DELEGATE_RetVal(
+	bool,
+	FTrackedTouchEnded)
+
+UCLASS(MinimalAPI)
+class UUIFrameworkButtonWidget : public UButton
+{
+	GENERATED_BODY()
+
+public:
+	UE_API UUIFrameworkButtonWidget();
+
+	// UWidget interface
+	UE_API virtual TSharedRef<SWidget> RebuildWidget() override;
+	UE_API virtual void ReleaseSlateResources(bool bReleaseChildren) override;
+	// End of UWidget interface
+
+	FSimpleDelegate OnTouchLongPress;
+	FSimpleDelegate OnTouchStarted;
+	FSimpleDelegate OnTouchEnded;
+
+protected:
+	UE_API FReply SlateHandleTouchLongPress(const FPointerEvent& PointerEvent);
+	UE_API FReply SlateHandleTouchStarted(const FPointerEvent& PointerEvent);
+	UE_API FReply SlateHandleTouchEnded(const FPointerEvent& PointerEvent);
+
+private:
+	void StartTrackingTouch(int32 PointerIndex);
+	void StopTrackingTouch();
+
+	bool HandleTouchEnded();
+
+private:
+	/** Cached pointer to the underlying slate button owned by this UWidget */
+	TSharedPtr<class SUIFrameworkButtonWidget> UIFButton;
+
+	TSharedPtr<FUIFInputPreProcessor> InputPreProcessor;
+	bool bTouchStarted = false;
+	bool bLongPressActive = false;
+};
+
+class FUIFInputPreProcessor
+	: public IInputProcessor
+	, public TSharedFromThis<FUIFInputPreProcessor>
+{
+public:
+	virtual void Tick(const float DeltaTime, FSlateApplication& SlateApp, TSharedRef<ICursor> Cursor) override {}
+	virtual bool HandleMouseButtonUpEvent(FSlateApplication& SlateApp, const FPointerEvent& MouseEvent) override;
+
+	FTrackedTouchEnded& OnTrackedTouchEnded() { return OnTrackedTouchEndedDelegate; }
+
+	void StartTrackingTouch(int32 Index);
+	void StopTrackingTouch();
+
+	int32 GetTrackedTouchIndex() const { return TrackedTouchIndex; }
+
+private:
+	int32 TrackedTouchIndex = INDEX_NONE;
+	FTrackedTouchEnded OnTrackedTouchEndedDelegate;
+};
+
+#undef UE_API
diff --git a/Engine/Plugins/Experimental/UIFramework/Source/Public/Types/UIFEvents.h b/Engine/Plugins/Experimental/UIFramework/Source/Public/Types/UIFEvents.h
--- a/Engine/Plugins/Experimental/UIFramework/Source/Public/Types/UIFEvents.h
+++ b/Engine/Plugins/Experimental/UIFramework/Source/Public/Types/UIFEvents.h
@@ -45,24 +45,24 @@ DECLARE_MULTICAST_DELEGATE_OneParam(FUIFrameworkClickEvent, FUIFrameworkClickEve
  *
  */
 USTRUCT(BlueprintType)
-struct FUIFrameworkFocusEventArgument : public FUIFrameworkSimpleEventArgument
+struct FUIFrameworkHighlightEventArgument : public FUIFrameworkSimpleEventArgument
 {
 	GENERATED_BODY()
 
-	FUIFrameworkFocusEventArgument() = default;
+	FUIFrameworkHighlightEventArgument() = default;
 };
 
-DECLARE_MULTICAST_DELEGATE_OneParam(FUIFrameworkFocusEvent, FUIFrameworkFocusEventArgument);
+DECLARE_MULTICAST_DELEGATE_OneParam(FUIFrameworkHighlightEvent, FUIFrameworkHighlightEventArgument);
 
 /**
  *
  */
 USTRUCT(BlueprintType)
-struct FUIFrameworkHoverEventArgument : public FUIFrameworkSimpleEventArgument
+struct FUIFrameworkUnhighlightEventArgument : public FUIFrameworkSimpleEventArgument
 {
 	GENERATED_BODY()
 
-	FUIFrameworkHoverEventArgument() = default;
+	FUIFrameworkUnhighlightEventArgument() = default;
 };
 
-DECLARE_MULTICAST_DELEGATE_OneParam(FUIFrameworkHoverEvent, FUIFrameworkHoverEventArgument);
\ No newline at end of file
+DECLARE_MULTICAST_DELEGATE_OneParam(FUIFrameworkUnhighlightEvent, FUIFrameworkUnhighlightEventArgument);
\ No newline at end of file
diff --git a/Engine/Plugins/Experimental/UIFramework/Source/Public/Widgets/UIFButton.h b/Engine/Plugins/Experimental/UIFramework/Source/Public/Widgets/UIFButton.h
--- a/Engine/Plugins/Experimental/UIFramework/Source/Public/Widgets/UIFButton.h
+++ b/Engine/Plugins/Experimental/UIFramework/Source/Public/Widgets/UIFButton.h
@@ -2,7 +2,6 @@
 
 #pragma once
 
-#include "Components/Button.h"
 #include "Types/UIFSlotBase.h"
 #include "Types/UIFEvents.h"
 #include "UIFWidget.h"
@@ -12,6 +11,7 @@
 #define UE_API UIFRAMEWORK_API
 
 struct FUIFrameworkWidgetId;
+enum class ECommonInputType : uint8;
 
 /**
  *
@@ -48,28 +48,22 @@ private:
 	UE_API void ServerClick(APlayerController* PlayerController);
 
 	UFUNCTION()
-	UE_API void HandleFocusReceived();
+	UE_API void HandleHighlight();
 
 	UFUNCTION(Server, Reliable)
-	UE_API void ServerFocusReceived(APlayerController* PlayerController);
+	UE_API void ServerHighlight(APlayerController* PlayerController);
 
 	UFUNCTION()
-	UE_API void HandleFocusLost();
+	UE_API void HandleUnhighlight();
 
 	UFUNCTION(Server, Reliable)
-	UE_API void ServerFocusLost(APlayerController* PlayerController);
+	UE_API void ServerUnhighlight(APlayerController* PlayerController);
 
 	UFUNCTION()
-	UE_API void HandleHovered();
-
-	UFUNCTION(Server, Reliable)
-	UE_API void ServerHovered(APlayerController* PlayerController);
+	UE_API void HandleHover();
 
 	UFUNCTION()
-	UE_API void HandleUnhovered();
-
-	UFUNCTION(Server, Reliable)
-	UE_API void ServerUnhovered(APlayerController* PlayerController);
+	UE_API void HandleUnhover();
 
 	UFUNCTION()
 	UE_API void OnRep_Slot();
@@ -77,25 +71,15 @@ private:
 public:
 	FUIFrameworkClickEvent OnClick;
 
-	FUIFrameworkFocusEvent OnFocusReceived;
-
-	FUIFrameworkFocusEvent OnFocusLost;
+	FUIFrameworkHighlightEvent OnHighlight;
 
-	FUIFrameworkHoverEvent OnHovered;
+	FUIFrameworkUnhighlightEvent OnUnhighlight;
 
-	FUIFrameworkHoverEvent OnUnhovered;
 private:
 	UPROPERTY(/*ExposeOnSpawn, */ReplicatedUsing = OnRep_Slot)
 	FUIFrameworkSimpleSlot Slot;
-};
 
-UCLASS(MinimalAPI)
-class UUIFrameworkButtonWidget : public UButton
-{
-	GENERATED_BODY()
-
-public:
-	UE_API UUIFrameworkButtonWidget();
+	bool bIsHighlighted = false;
 };
 
 #undef UE_API
diff --git a/Engine/Plugins/Experimental/UIFramework/Source/UIFramework.Build.cs b/Engine/Plugins/Experimental/UIFramework/Source/UIFramework.Build.cs
--- a/Engine/Plugins/Experimental/UIFramework/Source/UIFramework.Build.cs
+++ b/Engine/Plugins/Experimental/UIFramework/Source/UIFramework.Build.cs
@@ -31,6 +31,8 @@ public class UIFramework : ModuleRules
 			new string[]
 			{
 				"NetCore",
+				"CommonInput",
+				"CommonUI",
 			}
 		);
 
diff --git a/Engine/Plugins/Experimental/UIFramework/UIFramework.uplugin b/Engine/Plugins/Experimental/UIFramework/UIFramework.uplugin
--- a/Engine/Plugins/Experimental/UIFramework/UIFramework.uplugin
+++ b/Engine/Plugins/Experimental/UIFramework/UIFramework.uplugin
@@ -29,6 +29,10 @@
 		{
 			"Name": "LocalizableMessage",
 			"Enabled": true
+		},
+		{
+			"Name": "CommonUI",
+			"Enabled": true
 		}
 	]
 }
\ No newline at end of file
