UIF button: add triggering input action to the button API.
Clicking on the button will fire the input action and vice versa.
#Synced-CL 49373475

diff --git a/Engine/Plugins/Experimental/UIFramework/Source/Private/Widgets/UIFButton.cpp b/Engine/Plugins/Experimental/UIFramework/Source/Private/Widgets/UIFButton.cpp
--- a/Engine/Plugins/Experimental/UIFramework/Source/Private/Widgets/UIFButton.cpp
+++ b/Engine/Plugins/Experimental/UIFramework/Source/Private/Widgets/UIFButton.cpp
@@ -9,6 +9,7 @@
 #include "CommonInputSubsystem.h"
 #include "Components/Button.h"
 #include "Components/ButtonSlot.h"
+#include "Engine/AssetManager.h"
 #include "Misc/RedirectCollector.h"
 #include "Net/UnrealNetwork.h"
 
@@ -35,6 +36,7 @@ void UUIFrameworkButton::GetLifetimeReplicatedProps(TArray<FLifetimeProperty>& O
 	FDoRepLifetimeParams Params;
 	Params.bIsPushBased = true;
 	DOREPLIFETIME_WITH_PARAMS_FAST(ThisClass, Slot, Params);
+	DOREPLIFETIME_WITH_PARAMS_FAST(ThisClass, TriggeringInputAction, Params);
 }
 
 void UUIFrameworkButton::SetContent(FUIFrameworkSimpleSlot InEntry)
@@ -85,11 +87,12 @@ void UUIFrameworkButton::LocalOnUMGWidgetCreated()
 {
 	Super::LocalOnUMGWidgetCreated();
 	UUIFrameworkButtonWidget* Button = CastChecked<UUIFrameworkButtonWidget>(LocalGetUMGWidget());
-	Button->OnClicked.AddUniqueDynamic(this, &ThisClass::HandleClick);
+	Button->OnTriggered.AddUniqueDynamic(this, &ThisClass::HandleClick);
 	Button->OnHovered.AddUniqueDynamic(this, &ThisClass::HandleHover);
 	Button->OnUnhovered.AddUniqueDynamic(this, &ThisClass::HandleUnhover);
 	Button->OnTouchLongPress.BindUObject(this, &ThisClass::HandleHighlight);
 	Button->OnTouchEnded.BindUObject(this, &ThisClass::HandleUnhighlight);
+	LoadInputAction();
 }
 
 void UUIFrameworkButton::LocalAddChild(FUIFrameworkWidgetId ChildId)
@@ -120,6 +123,59 @@ void UUIFrameworkButton::LocalAddChild(FUIFrameworkWidgetId ChildId)
 	}
 }
 
+void UUIFrameworkButton::SetTriggeringInputAction(TSoftObjectPtr<UInputAction> InputAction)
+{
+	if (InputAction != TriggeringInputAction)
+	{
+		TriggeringInputAction = InputAction;
+		MARK_PROPERTY_DIRTY_FROM_NAME(ThisClass, TriggeringInputAction, this);
+	}
+}
+
+void UUIFrameworkButton::OnRep_TriggeringInputAction()
+{
+	LoadInputAction();
+}
+
+void UUIFrameworkButton::HandleInputActionLoaded()
+{
+#if !UE_SERVER
+	if (UUIFrameworkButtonWidget* Button = Cast<UUIFrameworkButtonWidget>(LocalGetUMGWidget()))
+	{
+		Button->SetTriggeringEnhancedInputAction(TriggeringInputAction.Get());
+	}
+	AsyncLoadInputActionHandle.Reset();
+#endif
+}
+
+void UUIFrameworkButton::LoadInputAction()
+{
+#if !UE_SERVER
+	if (AsyncLoadInputActionHandle.IsValid())
+	{
+		AsyncLoadInputActionHandle->CancelHandle();
+	}
+	AsyncLoadInputActionHandle.Reset();
+
+	if (!TriggeringInputAction.IsNull())
+	{
+		if (TriggeringInputAction.IsValid())
+		{
+			HandleInputActionLoaded();
+		}
+		else
+		{
+			FStreamableDelegate Delegate = FStreamableDelegate::CreateUObject(this, &UUIFrameworkButton::HandleInputActionLoaded);
+			AsyncLoadInputActionHandle = UAssetManager::GetStreamableManager().RequestAsyncLoad(TriggeringInputAction.ToSoftObjectPath(), Delegate);
+		}
+	}
+	else
+	{
+		HandleInputActionLoaded();
+	}
+#endif
+}
+
 void UUIFrameworkButton::HandleClick()
 {
 	// todo the click event should send the userid
diff --git a/Engine/Plugins/Experimental/UIFramework/Source/Private/Widgets/UIFrameworkButtonWidget.cpp b/Engine/Plugins/Experimental/UIFramework/Source/Private/Widgets/UIFrameworkButtonWidget.cpp
--- a/Engine/Plugins/Experimental/UIFramework/Source/Private/Widgets/UIFrameworkButtonWidget.cpp
+++ b/Engine/Plugins/Experimental/UIFramework/Source/Private/Widgets/UIFrameworkButtonWidget.cpp
@@ -4,9 +4,18 @@
 
 #include "SUIFrameworkButtonWidget.h"
 
+#include "CommonActivatableWidget.h"
+#include "CommonUITypes.h"
 #include "Components/Button.h"
 #include "Components/ButtonSlot.h"
+#include "EnhancedActionKeyMapping.h"
+#include "EnhancedInputSubsystems.h"
 #include "Framework/Application/SlateApplication.h"
+#include "GameFramework/PlayerController.h"
+#include "Input/CommonUIActionRouterBase.h"
+#include "Input/CommonUIInputTypes.h"
+#include "TimerManager.h"
+#include "UIFWidget.h"
 #include "Widgets/DeclarativeSyntaxSupport.h"
 
 #include UE_INLINE_GENERATED_CPP_BY_NAME(UIFrameworkButtonWidget)
@@ -46,6 +55,8 @@ TSharedRef<SWidget> UUIFrameworkButtonWidget::RebuildWidget()
 		.IsFocusable(GetIsFocusable())
 		.AllowDragDrop(bAllowDragDrop);
 
+	OnClicked.AddUniqueDynamic(this, &UUIFrameworkButtonWidget::HandleButtonClicked);
+
 	if (GetChildrenCount() > 0)
 	{
 		Cast<UButtonSlot>(GetContentSlot())->BuildSlot(MyButton.ToSharedRef());
@@ -60,9 +71,21 @@ TSharedRef<SWidget> UUIFrameworkButtonWidget::RebuildWidget()
 	return UIFButton.ToSharedRef();
 }
 
+void UUIFrameworkButtonWidget::OnWidgetRebuilt()
+{
+	Super::OnWidgetRebuilt();
+	if (UWorld* World = GetWorld())
+	{
+		FTimerManager& TimerManager = World->GetTimerManager();
+		TimerManager.SetTimerForNextTick(FTimerDelegate::CreateUObject(this, &UUIFrameworkButtonWidget::BindTriggeringInputActionToClick));
+	}
+}
+
 void UUIFrameworkButtonWidget::ReleaseSlateResources(bool bReleaseChildren)
 {
 	Super::ReleaseSlateResources(bReleaseChildren);
+
+	UnbindTriggeringInputActionToClick();
 	UIFButton.Reset();
 
 	if (InputPreProcessor)
@@ -74,6 +97,108 @@ void UUIFrameworkButtonWidget::ReleaseSlateResources(bool bReleaseChildren)
 	InputPreProcessor.Reset();
 }
 
+void UUIFrameworkButtonWidget::HandleButtonClicked()
+{
+	OnTriggered.Broadcast();
+	ExecuteTriggeredInput();
+}
+
+void UUIFrameworkButtonWidget::ExecuteTriggeredInput()
+{
+	// Fire all the bound input action events
+	if (TriggeringInputAction)
+	{
+		if (APlayerController* PlayerController = GetOwningPlayer())
+		{
+			if (ULocalPlayer* LocalPlayer = PlayerController->GetLocalPlayer())
+			{
+				if (UEnhancedInputLocalPlayerSubsystem* EnhancedInputLocalPlayerSubsystem = LocalPlayer->GetSubsystem<UEnhancedInputLocalPlayerSubsystem>())
+				{
+					if (const UEnhancedPlayerInput* const PlayerInput = EnhancedInputLocalPlayerSubsystem->GetPlayerInput())
+					{
+						// Confirm that this input action is active before injecting it
+						if (PlayerInput->FindActionInstanceData(TriggeringInputAction.Get()))
+						{
+							FInputActionValue RawValue = FInputActionValue(true);
+							CommonUI::InjectEnhancedInputForAction(LocalPlayer, TriggeringInputAction.Get(), RawValue);
+						}
+					}
+				}
+			}
+		}
+	}
+}
+
+void UUIFrameworkButtonWidget::BindTriggeringInputActionToClick()
+{
+	if (!TriggeringInputAction)
+	{
+		return;
+	}
+
+	if (!TriggeringBindingHandle.IsValid())
+	{
+		if (APlayerController* PlayerController = GetOwningPlayer())
+		{
+			TSharedPtr<SWidget> CurWidget = GetCachedWidget();
+			ULocalPlayer* LocalPlayer = PlayerController->GetLocalPlayer();
+
+			if (!LocalPlayer || !CurWidget)
+			{
+				return;
+			}
+
+			// Register a UI Action Binding for the input action so that it can be triggered in Menu mode
+			if (const UCommonActivatableWidget* ActivatableWidget = UCommonUIActionRouterBase::FindOwningActivatable(CurWidget, LocalPlayer))
+			{
+				if (UCommonUIActionRouterBase* ActionRouter = UCommonUIActionRouterBase::Get(*ActivatableWidget))
+				{
+					FBindUIActionArgs BindArgs(TriggeringInputAction.Get(), false, FSimpleDelegate::CreateUObject(this, &UUIFrameworkButtonWidget::HandleButtonClicked));
+					BindArgs.bConsumeInput = false;
+					TriggeringBindingHandle = ActionRouter->RegisterUIActionBinding(*this, BindArgs);
+				}
+			}
+		}
+	}
+}
+
+APlayerController* UUIFrameworkButtonWidget::GetOwningPlayer() const
+{
+	if (APlayerController* PlayerController = Super::GetOwningPlayer())
+	{
+		return PlayerController;
+	}
+	else if (UUIFrameworkWidget* FrameworkWidget = Cast<UUIFrameworkWidget>(GetOuter()))
+	{
+		return FrameworkWidget->GetPlayerController();
+	}
+
+	return nullptr;
+}
+
+void UUIFrameworkButtonWidget::UnbindTriggeringInputActionToClick()
+{
+	if (TriggeringBindingHandle.IsValid())
+	{
+		TriggeringBindingHandle.Unregister();
+	}
+}
+
+void UUIFrameworkButtonWidget::SetTriggeringEnhancedInputAction(UInputAction* InInputAction)
+{
+	if (TriggeringInputAction != InInputAction)
+	{
+		UnbindTriggeringInputActionToClick();
+
+		TriggeringInputAction = InInputAction;
+
+		if (!IsDesignTime())
+		{
+			BindTriggeringInputActionToClick();
+		}
+	}
+}
+
 FReply UUIFrameworkButtonWidget::SlateHandleTouchLongPress(const FPointerEvent& PointerEvent)
 {
 	if (UIFButton)
diff --git a/Engine/Plugins/Experimental/UIFramework/Source/Private/Widgets/UIFrameworkButtonWidget.h b/Engine/Plugins/Experimental/UIFramework/Source/Private/Widgets/UIFrameworkButtonWidget.h
--- a/Engine/Plugins/Experimental/UIFramework/Source/Private/Widgets/UIFrameworkButtonWidget.h
+++ b/Engine/Plugins/Experimental/UIFramework/Source/Private/Widgets/UIFrameworkButtonWidget.h
@@ -5,6 +5,7 @@
 #include "Components/Button.h"
 #include "Delegates/Delegate.h"
 #include "Framework/Application/IInputProcessor.h"
+#include "Input/UIActionBindingHandle.h"
 #include "UIFrameworkButtonWidget.generated.h"
 
 #define UE_API UIFRAMEWORK_API
@@ -15,6 +16,9 @@ DECLARE_DELEGATE_RetVal(
 	bool,
 	FTrackedTouchEnded)
 
+class UInputAction;
+class UPlayerController;
+
 UCLASS(MinimalAPI)
 class UUIFrameworkButtonWidget : public UButton
 {
@@ -25,28 +29,48 @@ public:
 
 	// UWidget interface
 	UE_API virtual TSharedRef<SWidget> RebuildWidget() override;
+	UE_API virtual void OnWidgetRebuilt() override;
 	UE_API virtual void ReleaseSlateResources(bool bReleaseChildren) override;
+	UE_API virtual APlayerController* GetOwningPlayer() const override;
 	// End of UWidget interface
 
+	void SetTriggeringEnhancedInputAction(UInputAction* InInputAction);
+
 	FSimpleDelegate OnTouchLongPress;
 	FSimpleDelegate OnTouchStarted;
 	FSimpleDelegate OnTouchEnded;
 
+	FOnButtonClickedEvent OnTriggered;
+
 protected:
 	UE_API FReply SlateHandleTouchLongPress(const FPointerEvent& PointerEvent);
 	UE_API FReply SlateHandleTouchStarted(const FPointerEvent& PointerEvent);
 	UE_API FReply SlateHandleTouchEnded(const FPointerEvent& PointerEvent);
 
 private:
+
+	UFUNCTION()
+	void HandleButtonClicked();
+	void ExecuteTriggeredInput();
+
+	void BindTriggeringInputActionToClick();
+	void UnbindTriggeringInputActionToClick();
+
 	void StartTrackingTouch(int32 PointerIndex);
 	void StopTrackingTouch();
 
 	bool HandleTouchEnded();
 
+protected:
+	UPROPERTY(EditAnywhere, BlueprintReadOnly, Category = Input)
+	TObjectPtr<UInputAction> TriggeringInputAction;
+
 private:
 	/** Cached pointer to the underlying slate button owned by this UWidget */
 	TSharedPtr<class SUIFrameworkButtonWidget> UIFButton;
 
+	FUIActionBindingHandle TriggeringBindingHandle;
+
 	TSharedPtr<FUIFInputPreProcessor> InputPreProcessor;
 	bool bTouchStarted = false;
 	bool bLongPressActive = false;
diff --git a/Engine/Plugins/Experimental/UIFramework/Source/Public/Widgets/UIFButton.h b/Engine/Plugins/Experimental/UIFramework/Source/Public/Widgets/UIFButton.h
--- a/Engine/Plugins/Experimental/UIFramework/Source/Public/Widgets/UIFButton.h
+++ b/Engine/Plugins/Experimental/UIFramework/Source/Public/Widgets/UIFButton.h
@@ -12,6 +12,7 @@
 
 struct FUIFrameworkWidgetId;
 enum class ECommonInputType : uint8;
+class UInputAction;
 
 /**
  *
@@ -33,6 +34,8 @@ public:
 		return Slot;
 	}
 
+	UE_API void SetTriggeringInputAction(TSoftObjectPtr<UInputAction> InputAction);
+
 	UE_API virtual void AuthorityForEachChildren(const TFunctionRef<void(UUIFrameworkWidget*)>& Func) override;
 	UE_API virtual void AuthorityRemoveChild(UUIFrameworkWidget* Widget) override;
 	UE_API virtual void LocalAddChild(FUIFrameworkWidgetId ChildId) override;
@@ -68,6 +71,12 @@ private:
 	UFUNCTION()
 	UE_API void OnRep_Slot();
 
+	UFUNCTION()
+	UE_API void OnRep_TriggeringInputAction();
+
+	void LoadInputAction();
+	void HandleInputActionLoaded();
+
 public:
 	FUIFrameworkClickEvent OnClick;
 
@@ -80,6 +89,11 @@ private:
 	FUIFrameworkSimpleSlot Slot;
 
 	bool bIsHighlighted = false;
+
+	TSharedPtr<FStreamableHandle> AsyncLoadInputActionHandle;
+
+	UPROPERTY(ReplicatedUsing = OnRep_TriggeringInputAction)
+	TSoftObjectPtr<UInputAction> TriggeringInputAction;
 };
 
 #undef UE_API
diff --git a/Engine/Plugins/Experimental/UIFramework/Source/UIFramework.Build.cs b/Engine/Plugins/Experimental/UIFramework/Source/UIFramework.Build.cs
--- a/Engine/Plugins/Experimental/UIFramework/Source/UIFramework.Build.cs
+++ b/Engine/Plugins/Experimental/UIFramework/Source/UIFramework.Build.cs
@@ -9,10 +9,12 @@ public class UIFramework : ModuleRules
 		PublicDependencyModuleNames.AddRange(
 			new string[]
 			{
+				"CommonUI",
 				"Core",
 				"CoreUObject",
 				"DeveloperSettings",
 				"Engine",
+				"EnhancedInput",
 				"FieldNotification",
 				"LocalizableMessage",
 				"ModelViewViewModel",
@@ -32,7 +34,7 @@ public class UIFramework : ModuleRules
 			{
 				"NetCore",
 				"CommonInput",
-				"CommonUI",
+				"GameplayTags",
 			}
 		);
 
