Fix for ensure in verse widgets when trying to instantiate a widget while the world is tearing down

#Synced-CL 49588415

diff --git a/Engine/Plugins/Experimental/UIFramework/Source/Private/UIFPlayerComponent.cpp b/Engine/Plugins/Experimental/UIFramework/Source/Private/UIFPlayerComponent.cpp
--- a/Engine/Plugins/Experimental/UIFramework/Source/Private/UIFPlayerComponent.cpp
+++ b/Engine/Plugins/Experimental/UIFramework/Source/Private/UIFPlayerComponent.cpp
@@ -126,6 +126,15 @@ void UUIFrameworkPlayerComponent::UninitializeComponent()
 		FTSTicker::RemoveTicker(PendingWidgetConstructedTickerHandle);
 	}
 
+	for (const auto& ClassesToLoadPair : ClassesToLoad)
+	{
+		if (ClassesToLoadPair.Value.StreamableHandle.IsValid())
+		{
+			ClassesToLoadPair.Value.StreamableHandle->CancelHandle();
+		}
+	}
+	ClassesToLoad.Empty();
+
 	Super::UninitializeComponent();
 }
 
