Fixing a bug where material_blocks break when the underlying material isn't loaded yet.

#Synced-CL 49588415

diff --git a/Engine/Plugins/Experimental/UIFramework/Source/Public/Widgets/UIFImageBlock.h b/Engine/Plugins/Experimental/UIFramework/Source/Public/Widgets/UIFImageBlock.h
--- a/Engine/Plugins/Experimental/UIFramework/Source/Public/Widgets/UIFImageBlock.h
+++ b/Engine/Plugins/Experimental/UIFramework/Source/Public/Widgets/UIFImageBlock.h
@@ -97,12 +97,13 @@ public:
 	UE_API virtual void LocalOnUMGWidgetCreated() override;
 	UE_API virtual bool LocalIsReplicationReady() const override;
 
+protected:
+	UE_API virtual UObject* ProcessResourceObject();
+
 private:
 	UFUNCTION()
 	UE_API void OnRep_Data();
 
-	UE_API UObject* ProcessResourceObject();
-
 private:
 	/** Tinting applied to the image. */
 	UPROPERTY(ReplicatedUsing = OnRep_Data)
