UIFramework button: remove padding when button is pressed, causing a visual jiggle
#Synced-CL 49593022

diff --git a/Engine/Plugins/Experimental/UIFramework/Source/Private/Widgets/UIFrameworkButtonWidget.cpp b/Engine/Plugins/Experimental/UIFramework/Source/Private/Widgets/UIFrameworkButtonWidget.cpp
--- a/Engine/Plugins/Experimental/UIFramework/Source/Private/Widgets/UIFrameworkButtonWidget.cpp
+++ b/Engine/Plugins/Experimental/UIFramework/Source/Private/Widgets/UIFrameworkButtonWidget.cpp
@@ -24,6 +24,7 @@ UUIFrameworkButtonWidget::UUIFrameworkButtonWidget()
 {
 	FButtonStyle TempStyle = GetStyle();
 	TempStyle.NormalPadding = FMargin(0.0f);
+	TempStyle.PressedPadding = FMargin(0.0f);
 	TempStyle.Normal.DrawAs = ESlateBrushDrawType::NoDrawType;
 	TempStyle.Hovered.DrawAs = ESlateBrushDrawType::NoDrawType;
 	TempStyle.Pressed.DrawAs = ESlateBrushDrawType::NoDrawType;
