Make material_block and umg_widget work with material root object replication

#Synced-CL 49823132

diff --git a/Engine/Plugins/Experimental/UIFramework/Source/Public/Widgets/UIFImageBlock.h b/Engine/Plugins/Experimental/UIFramework/Source/Public/Widgets/UIFImageBlock.h
--- a/Engine/Plugins/Experimental/UIFramework/Source/Public/Widgets/UIFImageBlock.h
+++ b/Engine/Plugins/Experimental/UIFramework/Source/Public/Widgets/UIFImageBlock.h
@@ -56,7 +56,7 @@ public:
 	UE_API UUIFrameworkImageBlock();
 
 	UFUNCTION(BlueprintCallable, BlueprintAuthorityOnly, Category = "UI Framework")
-	UE_API void SetMaterial(TSoftObjectPtr<UMaterialInterface> Material);
+	UE_API virtual void SetMaterial(TSoftObjectPtr<UMaterialInterface> Material);
 
 	UFUNCTION(BlueprintCallable, Category = "UI Framework")
 	TSoftObjectPtr<UObject> GeResourceObject() const
