UIFText: add wrapping
#Synced-CL 49794217
#Synced-CL 49794217

#Synced-CL 49794217

diff --git a/Engine/Plugins/Experimental/UIFramework/Source/Private/Widgets/UIFTextBlock.cpp b/Engine/Plugins/Experimental/UIFramework/Source/Private/Widgets/UIFTextBlock.cpp
--- a/Engine/Plugins/Experimental/UIFramework/Source/Private/Widgets/UIFTextBlock.cpp
+++ b/Engine/Plugins/Experimental/UIFramework/Source/Private/Widgets/UIFTextBlock.cpp
@@ -32,6 +32,7 @@ void UUIFrameworkTextBase::GetLifetimeReplicatedProps(TArray<FLifetimeProperty>&
 	DOREPLIFETIME_WITH_PARAMS_FAST(ThisClass, TextColor, Params);
 	DOREPLIFETIME_WITH_PARAMS_FAST(ThisClass, Justification, Params);
 	DOREPLIFETIME_WITH_PARAMS_FAST(ThisClass, OverflowPolicy, Params);
+	DOREPLIFETIME_WITH_PARAMS_FAST(ThisClass, Wrapping, Params);
 }
 
 void UUIFrameworkTextBase::SetText(const FLocalizableMessage& InMessage)
@@ -74,6 +75,16 @@ void UUIFrameworkTextBase::SetTextSize(float InTextSize)
 	}
 }
 
+void UUIFrameworkTextBase::SetWrapping(bool InWrapping)
+{
+	if (Wrapping != InWrapping)
+	{
+		Wrapping = InWrapping;
+		MARK_PROPERTY_DIRTY_FROM_NAME(ThisClass, Wrapping, this);
+		ForceNetUpdate();
+	}
+}
+
 void UUIFrameworkTextBase::SetJustification(ETextJustify::Type InJustification)
 {
 	if (Justification != InJustification)
@@ -100,6 +111,7 @@ void UUIFrameworkTextBase::LocalOnUMGWidgetCreated()
 	SetTextToWidget(Text);
 	SetTextColorToWidget(TextColor);
 	SetTextSizeToWidget(TextSize);
+	SetWrappingToWidget(Wrapping);
 	SetJustificationToWidget(Justification);
 	SetOverflowPolicyToWidget(OverflowPolicy);
 }
@@ -130,6 +142,14 @@ void UUIFrameworkTextBase::OnRep_TextSize()
 	}
 }
 
+void UUIFrameworkTextBase::OnRep_Wrapping()
+{
+	if (LocalGetUMGWidget())
+	{
+		SetWrappingToWidget(Wrapping);
+	}
+}
+
 void UUIFrameworkTextBase::OnRep_Justification()
 {
 	if (LocalGetUMGWidget())
@@ -220,6 +240,11 @@ void UUIFrameworkTextBlock::SetTextSizeToWidget(float InValue)
 	TextBlock->SetFont(CurrentFontInfo);
 }
 
+void UUIFrameworkTextBlock::SetWrappingToWidget(bool InValue)
+{
+	CastChecked<UTextBlock>(LocalGetUMGWidget())->SetAutoWrapText(InValue);
+}
+
 void UUIFrameworkTextBlock::SetTextColorToWidget(FLinearColor InValue)
 {
 	CastChecked<UTextBlock>(LocalGetUMGWidget())->SetColorAndOpacity(InValue);
diff --git a/Engine/Plugins/Experimental/UIFramework/Source/Public/Widgets/UIFTextBlock.h b/Engine/Plugins/Experimental/UIFramework/Source/Public/Widgets/UIFTextBlock.h
--- a/Engine/Plugins/Experimental/UIFramework/Source/Public/Widgets/UIFTextBlock.h
+++ b/Engine/Plugins/Experimental/UIFramework/Source/Public/Widgets/UIFTextBlock.h
@@ -52,6 +52,15 @@ public:
 		return TextSize;
 	}
 
+	UFUNCTION(BlueprintCallable, BlueprintAuthorityOnly, Category = "UI Framework")
+	UE_API void SetWrapping(bool InWrapping);
+
+	UFUNCTION(BlueprintCallable, Category = "UI Framework")
+	bool GetWrapping() const
+	{
+		return Wrapping;
+	}
+
 	UFUNCTION(BlueprintCallable, BlueprintAuthorityOnly, Category = "UI Framework")
 	UE_API void SetJustification(ETextJustify::Type Justification);
 
@@ -82,6 +91,9 @@ private:
 	UFUNCTION()
 	UE_API void OnRep_TextSize();
 
+	UFUNCTION()
+	UE_API void OnRep_Wrapping();
+
 	UFUNCTION()
 	UE_API void OnRep_Justification();
 
@@ -100,6 +112,9 @@ protected:
 	virtual void SetTextSizeToWidget(float)
 	{
 	}
+	virtual void SetWrappingToWidget(bool)
+	{
+	}
 	virtual void SetJustificationToWidget(ETextJustify::Type)
 	{
 	}
@@ -120,6 +135,9 @@ private:
 	UPROPERTY(ReplicatedUsing = OnRep_TextSize)
 	float TextSize = 24.0;
 
+	UPROPERTY(ReplicatedUsing = OnRep_Wrapping)
+	bool Wrapping = false;
+
 	UPROPERTY(ReplicatedUsing = OnRep_Justification)
 	TEnumAsByte<ETextJustify::Type> Justification;
 
@@ -170,6 +188,7 @@ protected:
 	UE_API virtual void SetOverflowPolicyToWidget(ETextOverflowPolicy) override;
 	UE_API virtual void SetTextColorToWidget(FLinearColor) override;
 	UE_API virtual void SetTextSizeToWidget(float) override;
+	UE_API virtual void SetWrappingToWidget(bool) override;
 	UE_API virtual void SetShadowOffsetToWidget(FVector2f);
 	UE_API virtual void SetShadowColorToWidget(FLinearColor);
 
