UIFPlayerComponent: fix SetFocus not working if called on the same widget multiple times
Local focus may have changed, so SetFocus needs to replicate the widget ID even if it hasn't changed.
#Synced-CL 49794217

diff --git a/Engine/Plugins/Experimental/UIFramework/Source/Private/UIFPlayerComponent.cpp b/Engine/Plugins/Experimental/UIFramework/Source/Private/UIFPlayerComponent.cpp
--- a/Engine/Plugins/Experimental/UIFramework/Source/Private/UIFPlayerComponent.cpp
+++ b/Engine/Plugins/Experimental/UIFramework/Source/Private/UIFPlayerComponent.cpp
@@ -259,16 +259,16 @@ void UUIFrameworkPlayerComponent::HandleAddPending()
 
 	PrimaryComponentTick.SetTickFunctionEnable(false);
 
-	if (WidgetToFocus.IsValid())
+	if (bIsPendingFocusCall && WidgetToFocus.Id.IsValid())
 	{
-		if (UUIFrameworkWidget* Widget = GetWidgetTree().FindWidgetById(WidgetToFocus))
+		if (UUIFrameworkWidget* Widget = GetWidgetTree().FindWidgetById(WidgetToFocus.Id))
 		{
 			if (UWidget* UMGWidget = Widget->LocalGetUMGWidget())
 			{
 				if (!PendingWidgetConstructedTickerHandle.IsValid())
 				{
-					PendingWidgetConstructedTickerHandle = FTSTicker::GetCoreTicker().AddTicker(FTickerDelegate::CreateUObject(this, &UUIFrameworkPlayerComponent::TrySetFocus, WidgetToFocus));
-					WidgetToFocus = FUIFrameworkWidgetId();
+					PendingWidgetConstructedTickerHandle = FTSTicker::GetCoreTicker().AddTicker(FTickerDelegate::CreateUObject(this, &UUIFrameworkPlayerComponent::TrySetFocus, WidgetToFocus.Id));
+					bIsPendingFocusCall = false;
 				}
 			}
 		}
@@ -412,7 +412,8 @@ void UUIFrameworkPlayerComponent::ServerRemoveWidgetRootFromTree_Implementation(
 
 void UUIFrameworkPlayerComponent::SetWidgetToFocus(FUIFrameworkWidgetId WidgetId)
 {
-	WidgetToFocus = WidgetId;
+	WidgetToFocus.Id = WidgetId;
+	WidgetToFocus.ReplicationCounter++;
 	MARK_PROPERTY_DIRTY_FROM_NAME(ThisClass, WidgetToFocus, this);
 }
 
@@ -495,15 +496,15 @@ void UUIFrameworkPlayerComponent::LocalAddChild(FUIFrameworkWidgetId WidgetId)
 	}
 }
 
-bool UUIFrameworkPlayerComponent::TrySetFocus(float DeltaTime, FUIFrameworkWidgetId CurrentWidgetToFocus)
+bool UUIFrameworkPlayerComponent::TrySetFocus(float DeltaTime, FUIFrameworkWidgetId CurrentWidgetToFocusId)
 {
-	if (UUIFrameworkWidget* Widget = GetWidgetTree().FindWidgetById(CurrentWidgetToFocus))
+	if (UUIFrameworkWidget* Widget = GetWidgetTree().FindWidgetById(CurrentWidgetToFocusId))
 	{
 		if (UWidget* UMGWidget = Widget->LocalGetUMGWidget())
 		{
 			if (UMGWidget->IsConstructed())
 			{
-				const FUIFrameworkWidgetTreeEntry* RootEntry = GetWidgetTree().FindRootEntryById(CurrentWidgetToFocus);
+				const FUIFrameworkWidgetTreeEntry* RootEntry = GetWidgetTree().FindRootEntryById(CurrentWidgetToFocusId);
 				if (RootEntry)
 				{
 					if (FUIFrameworkGameLayerSlot* LayerSlot = RootList.FindEntry(RootEntry->ChildId))
@@ -526,16 +527,17 @@ bool UUIFrameworkPlayerComponent::TrySetFocus(float DeltaTime, FUIFrameworkWidge
 
 void UUIFrameworkPlayerComponent::OnRep_WidgetToFocus()
 {
-	if (WidgetToFocus.IsValid())
+	if (WidgetToFocus.Id.IsValid())
 	{
-		if (UUIFrameworkWidget* Widget = GetWidgetTree().FindWidgetById(WidgetToFocus))
+		bIsPendingFocusCall = true;
+		if (UUIFrameworkWidget* Widget = GetWidgetTree().FindWidgetById(WidgetToFocus.Id))
 		{
 			if (UWidget* UMGWidget = Widget->LocalGetUMGWidget())
 			{
 				if (UMGWidget->IsConstructed())
 				{
-					TrySetFocus(0, WidgetToFocus);
-					WidgetToFocus = FUIFrameworkWidgetId();
+					TrySetFocus(0, WidgetToFocus.Id);
+					bIsPendingFocusCall = false;
 					return;
 				}
 			}
@@ -543,8 +545,8 @@ void UUIFrameworkPlayerComponent::OnRep_WidgetToFocus()
 			{
 				FTSTicker::RemoveTicker(PendingWidgetConstructedTickerHandle);
 			}
-			PendingWidgetConstructedTickerHandle = FTSTicker::GetCoreTicker().AddTicker(FTickerDelegate::CreateUObject(this, &UUIFrameworkPlayerComponent::TrySetFocus, WidgetToFocus));
-			WidgetToFocus = FUIFrameworkWidgetId();
+			PendingWidgetConstructedTickerHandle = FTSTicker::GetCoreTicker().AddTicker(FTickerDelegate::CreateUObject(this, &UUIFrameworkPlayerComponent::TrySetFocus, WidgetToFocus.Id));
+			bIsPendingFocusCall = false;
 		}
 	}
 }
\ No newline at end of file
diff --git a/Engine/Plugins/Experimental/UIFramework/Source/Public/UIFPlayerComponent.h b/Engine/Plugins/Experimental/UIFramework/Source/Public/UIFPlayerComponent.h
--- a/Engine/Plugins/Experimental/UIFramework/Source/Public/UIFPlayerComponent.h
+++ b/Engine/Plugins/Experimental/UIFramework/Source/Public/UIFPlayerComponent.h
@@ -63,6 +63,22 @@ struct FUIFrameworkGameLayerSlot : public FUIFrameworkSlotBase
 	EUIFrameworkGameLayerType Type = EUIFrameworkGameLayerType::Viewport;
 };
 
+/**
+ *
+ */
+USTRUCT()
+struct FUIFrameworkWidgetToFocus
+{
+	GENERATED_BODY()
+
+	//~ ID of the widget to focus on.
+	UPROPERTY()
+	FUIFrameworkWidgetId Id;
+
+	//~ Tracker to ensure widget ID is replicated even if it hasn't changed.
+	UPROPERTY()
+	int32 ReplicationCounter = 0;
+};
 
 /**
  *
@@ -190,7 +206,7 @@ private:
 
 	UFUNCTION()
 	UE_API void OnRep_WidgetToFocus();
-	UE_API bool TrySetFocus(float DeltaTime, FUIFrameworkWidgetId CurrentWidgetToFocus);
+	UE_API bool TrySetFocus(float DeltaTime, FUIFrameworkWidgetId CurrentWidgetToFocusId);
 
 	void HandleAddPending();
 
@@ -217,9 +233,11 @@ private:
 	//~ Once widgets are created and constructed, allow for actions such as focus to occur
 	FOnPendingReplicationProcessed OnPendingReplicationProcessed;
 
-	//~ Widget ID to focus on
+	//~ Widget to focus on
 	UPROPERTY(ReplicatedUsing=OnRep_WidgetToFocus, Transient)
-	FUIFrameworkWidgetId WidgetToFocus;
+	FUIFrameworkWidgetToFocus WidgetToFocus;
+
+	bool bIsPendingFocusCall = false;
 
 	struct FWidgetClassToLoad
 	{
