UIF text: Expose Wrapping Policy and Wrap Width
#Synced-CL 50152522
#Synced-CL 50152522

#Synced-CL 50152522

diff --git a/Engine/Plugins/Experimental/UIFramework/Source/Private/Widgets/UIFTextBlock.cpp b/Engine/Plugins/Experimental/UIFramework/Source/Private/Widgets/UIFTextBlock.cpp
--- a/Engine/Plugins/Experimental/UIFramework/Source/Private/Widgets/UIFTextBlock.cpp
+++ b/Engine/Plugins/Experimental/UIFramework/Source/Private/Widgets/UIFTextBlock.cpp
@@ -32,7 +32,9 @@ void UUIFrameworkTextBase::GetLifetimeReplicatedProps(TArray<FLifetimeProperty>&
 	DOREPLIFETIME_WITH_PARAMS_FAST(ThisClass, TextColor, Params);
 	DOREPLIFETIME_WITH_PARAMS_FAST(ThisClass, Justification, Params);
 	DOREPLIFETIME_WITH_PARAMS_FAST(ThisClass, OverflowPolicy, Params);
-	DOREPLIFETIME_WITH_PARAMS_FAST(ThisClass, Wrapping, Params);
+	DOREPLIFETIME_WITH_PARAMS_FAST(ThisClass, AutoWrap, Params);
+	DOREPLIFETIME_WITH_PARAMS_FAST(ThisClass, WrappingPolicy, Params);
+	DOREPLIFETIME_WITH_PARAMS_FAST(ThisClass, WrapWidth, Params);
 }
 
 void UUIFrameworkTextBase::SetText(const FLocalizableMessage& InMessage)
@@ -75,12 +77,32 @@ void UUIFrameworkTextBase::SetTextSize(float InTextSize)
 	}
 }
 
-void UUIFrameworkTextBase::SetWrapping(bool InWrapping)
+void UUIFrameworkTextBase::SetAutoWrap(bool InAutoWrap)
 {
-	if (Wrapping != InWrapping)
+	if (AutoWrap != InAutoWrap)
 	{
-		Wrapping = InWrapping;
-		MARK_PROPERTY_DIRTY_FROM_NAME(ThisClass, Wrapping, this);
+		AutoWrap = InAutoWrap;
+		MARK_PROPERTY_DIRTY_FROM_NAME(ThisClass, AutoWrap, this);
+		ForceNetUpdate();
+	}
+}
+
+void UUIFrameworkTextBase::SetWrapWidth(float InWrapWidth)
+{
+	if (WrapWidth != InWrapWidth)
+	{
+		WrapWidth = InWrapWidth;
+		MARK_PROPERTY_DIRTY_FROM_NAME(ThisClass, WrapWidth, this);
+		ForceNetUpdate();
+	}
+}
+
+void UUIFrameworkTextBase::SetWrappingPolicy(ETextWrappingPolicy InWrappingPolicy)
+{
+	if (WrappingPolicy != InWrappingPolicy)
+	{
+		WrappingPolicy = InWrappingPolicy;
+		MARK_PROPERTY_DIRTY_FROM_NAME(ThisClass, WrappingPolicy, this);
 		ForceNetUpdate();
 	}
 }
@@ -111,7 +133,9 @@ void UUIFrameworkTextBase::LocalOnUMGWidgetCreated()
 	SetTextToWidget(Text);
 	SetTextColorToWidget(TextColor);
 	SetTextSizeToWidget(TextSize);
-	SetWrappingToWidget(Wrapping);
+	SetAutoWrapToWidget(AutoWrap);
+	SetWrapWidthToWidget(WrapWidth);
+	SetWrappingPolicyToWidget(WrappingPolicy);
 	SetJustificationToWidget(Justification);
 	SetOverflowPolicyToWidget(OverflowPolicy);
 }
@@ -142,11 +166,27 @@ void UUIFrameworkTextBase::OnRep_TextSize()
 	}
 }
 
-void UUIFrameworkTextBase::OnRep_Wrapping()
+void UUIFrameworkTextBase::OnRep_AutoWrap()
+{
+	if (LocalGetUMGWidget())
+	{
+		SetAutoWrapToWidget(AutoWrap);
+	}
+}
+
+void UUIFrameworkTextBase::OnRep_WrapWidth()
 {
 	if (LocalGetUMGWidget())
 	{
-		SetWrappingToWidget(Wrapping);
+		SetWrapWidthToWidget(WrapWidth);
+	}
+}
+
+void UUIFrameworkTextBase::OnRep_WrappingPolicy()
+{
+	if (LocalGetUMGWidget())
+	{
+		SetWrappingPolicyToWidget(WrappingPolicy);
 	}
 }
 
@@ -240,11 +280,21 @@ void UUIFrameworkTextBlock::SetTextSizeToWidget(float InValue)
 	TextBlock->SetFont(CurrentFontInfo);
 }
 
-void UUIFrameworkTextBlock::SetWrappingToWidget(bool InValue)
+void UUIFrameworkTextBlock::SetAutoWrapToWidget(bool InValue)
 {
 	CastChecked<UTextBlock>(LocalGetUMGWidget())->SetAutoWrapText(InValue);
 }
 
+void UUIFrameworkTextBlock::SetWrapWidthToWidget(float InValue)
+{
+	CastChecked<UTextBlock>(LocalGetUMGWidget())->SetWrapTextAt(InValue);
+}
+
+void UUIFrameworkTextBlock::SetWrappingPolicyToWidget(ETextWrappingPolicy InValue)
+{
+	CastChecked<UTextBlock>(LocalGetUMGWidget())->SetWrappingPolicy(InValue);
+}
+
 void UUIFrameworkTextBlock::SetTextColorToWidget(FLinearColor InValue)
 {
 	CastChecked<UTextBlock>(LocalGetUMGWidget())->SetColorAndOpacity(InValue);
diff --git a/Engine/Plugins/Experimental/UIFramework/Source/Public/Widgets/UIFTextBlock.h b/Engine/Plugins/Experimental/UIFramework/Source/Public/Widgets/UIFTextBlock.h
--- a/Engine/Plugins/Experimental/UIFramework/Source/Public/Widgets/UIFTextBlock.h
+++ b/Engine/Plugins/Experimental/UIFramework/Source/Public/Widgets/UIFTextBlock.h
@@ -53,12 +53,30 @@ public:
 	}
 
 	UFUNCTION(BlueprintCallable, BlueprintAuthorityOnly, Category = "UI Framework")
-	UE_API void SetWrapping(bool InWrapping);
+	UE_API void SetAutoWrap(bool InAutoWrap);
 
 	UFUNCTION(BlueprintCallable, Category = "UI Framework")
-	bool GetWrapping() const
+	bool GetAutoWrap() const
 	{
-		return Wrapping;
+		return AutoWrap;
+	}
+
+	UFUNCTION(BlueprintCallable, BlueprintAuthorityOnly, Category = "UI Framework")
+	UE_API void SetWrapWidth(float InWrapWidth);
+
+	UFUNCTION(BlueprintCallable, Category = "UI Framework")
+	float GetWrapWidth() const
+	{
+		return WrapWidth;
+	}
+
+	UFUNCTION(BlueprintCallable, BlueprintAuthorityOnly, Category = "UI Framework")
+	UE_API void SetWrappingPolicy(ETextWrappingPolicy InWrappingPolicy);
+
+	UFUNCTION(BlueprintCallable, Category = "UI Framework")
+	ETextWrappingPolicy GetWrappingPolicy() const
+	{
+		return WrappingPolicy;
 	}
 
 	UFUNCTION(BlueprintCallable, BlueprintAuthorityOnly, Category = "UI Framework")
@@ -92,7 +110,13 @@ private:
 	UE_API void OnRep_TextSize();
 
 	UFUNCTION()
-	UE_API void OnRep_Wrapping();
+	UE_API void OnRep_AutoWrap();
+
+	UFUNCTION()
+	UE_API void OnRep_WrapWidth();
+
+	UFUNCTION()
+	UE_API void OnRep_WrappingPolicy();
 
 	UFUNCTION()
 	UE_API void OnRep_Justification();
@@ -112,7 +136,13 @@ protected:
 	virtual void SetTextSizeToWidget(float)
 	{
 	}
-	virtual void SetWrappingToWidget(bool)
+	virtual void SetAutoWrapToWidget(bool)
+	{
+	}
+	virtual void SetWrapWidthToWidget(float)
+	{
+	}
+	virtual void SetWrappingPolicyToWidget(ETextWrappingPolicy)
 	{
 	}
 	virtual void SetJustificationToWidget(ETextJustify::Type)
@@ -135,8 +165,14 @@ private:
 	UPROPERTY(ReplicatedUsing = OnRep_TextSize)
 	float TextSize = 24.0;
 
-	UPROPERTY(ReplicatedUsing = OnRep_Wrapping)
-	bool Wrapping = false;
+	UPROPERTY(ReplicatedUsing = OnRep_AutoWrap)
+	bool AutoWrap = false;
+
+	UPROPERTY(ReplicatedUsing = OnRep_WrapWidth)
+	float WrapWidth = 0.0f;
+
+	UPROPERTY(ReplicatedUsing = OnRep_WrappingPolicy)
+	ETextWrappingPolicy WrappingPolicy;
 
 	UPROPERTY(ReplicatedUsing = OnRep_Justification)
 	TEnumAsByte<ETextJustify::Type> Justification;
@@ -188,7 +224,9 @@ protected:
 	UE_API virtual void SetOverflowPolicyToWidget(ETextOverflowPolicy) override;
 	UE_API virtual void SetTextColorToWidget(FLinearColor) override;
 	UE_API virtual void SetTextSizeToWidget(float) override;
-	UE_API virtual void SetWrappingToWidget(bool) override;
+	UE_API virtual void SetAutoWrapToWidget(bool) override;
+	UE_API virtual void SetWrapWidthToWidget(float) override;
+	UE_API virtual void SetWrappingPolicyToWidget(ETextWrappingPolicy) override;
 	UE_API virtual void SetShadowOffsetToWidget(FVector2f);
 	UE_API virtual void SetShadowColorToWidget(FLinearColor);
 
