Add VeryVerbose logs to UIFramework for debugging

#Synced-CL 50417876

diff --git a/Engine/Plugins/Experimental/UIFramework/Source/Private/UIFPlayerComponent.cpp b/Engine/Plugins/Experimental/UIFramework/Source/Private/UIFPlayerComponent.cpp
--- a/Engine/Plugins/Experimental/UIFramework/Source/Private/UIFPlayerComponent.cpp
+++ b/Engine/Plugins/Experimental/UIFramework/Source/Private/UIFPlayerComponent.cpp
@@ -97,6 +97,8 @@ void UUIFrameworkPlayerComponent::InitializeComponent()
 {
 	Super::InitializeComponent();
 
+	UE_LOG(LogUIFramework, VeryVerbose, TEXT("Initializing UIFrameworkPlayerComponent"));
+
 	if (GetOwner()->HasAuthority())
 	{
 		WidgetTree.AuthorityAddAllWidgetsFromActorChannel();
@@ -109,6 +111,8 @@ void UUIFrameworkPlayerComponent::InitializeComponent()
 
 void UUIFrameworkPlayerComponent::UninitializeComponent()
 {
+	UE_LOG(LogUIFramework, VeryVerbose, TEXT("Uninitializing UIFrameworkPlayerComponent"));
+
 	// On local, remove all UWidget.
 	if (GetOwner()->HasAuthority())
 	{
@@ -170,6 +174,7 @@ void UUIFrameworkPlayerComponent::AddWidget(FUIFrameworkGameLayerSlot InEntry)
 	}
 	else
 	{
+		UE_LOG(LogUIFramework, VeryVerbose, TEXT("Attempting to add widget '%" INT64_FMT "' on the server."), InEntry.AuthorityGetWidget()->GetWidgetId().GetKey());
 		// Reset the widget to make sure the id is set and it may have been duplicated during the attach
 		InEntry.AuthoritySetWidget(FUIFrameworkModule::AuthorityAttachWidget(this, InEntry.AuthorityGetWidget()));
 		RootList.AddEntry(InEntry);
@@ -193,6 +198,8 @@ void UUIFrameworkPlayerComponent::RemoveWidget(UUIFrameworkWidget* Widget)
 		}
 		else
 		{
+			UE_LOG(LogUIFramework, VeryVerbose, TEXT("Attempting to remove widget '%" INT64_FMT "' on the server."), Widget->GetWidgetId().GetKey());
+
 			if (RootList.RemoveEntry(Widget))
 			{
 				FUIFrameworkModule::AuthorityDetachWidgetFromParent(Widget);
@@ -233,6 +240,7 @@ void UUIFrameworkPlayerComponent::HandleAddPending()
 	{
 		if (FUIFrameworkWidgetTreeEntry* Entry = WidgetTree.LocalGetEntryByReplicationId(ReplicationId))
 		{
+			UE_LOG(LogUIFramework, VeryVerbose, TEXT("Attempting to add pending widget '%" INT64_FMT "'."), Entry->ChildId.GetKey());
 			if (ensure(Entry->IsParentValid() && Entry->IsChildValid()))
 			{
 				if (Entry->ParentId.IsRoot())
@@ -314,6 +322,7 @@ void UUIFrameworkPlayerComponent::TickComponent(float DeltaTime, enum ELevelTick
 
 	if (AddPending.Num() > 0 && ClassesToLoad.Num() == 0 && NetReplicationPending.Num() == 0)
 	{
+		UE_LOG(LogUIFramework, VeryVerbose, TEXT("Pending widgets found, attempting to add."));
 		HandleAddPending();
 	}
 
@@ -326,9 +335,11 @@ void UUIFrameworkPlayerComponent::LocalWidgetWasAddedToTree(const FUIFrameworkWi
 
 	if (Entry.Child && Entry.Child->LocalIsReplicationReady())
 	{
+
 		TSoftClassPtr<UWidget> WidgetClass = Entry.Child->GetUMGWidgetClass();
 		if (WidgetClass.IsValid())
 		{
+			UE_LOG(LogUIFramework, VeryVerbose, TEXT("Widget %s with ID  '%" INT64_FMT "' is ready to be created on the client."), *WidgetClass->GetName(), Entry.ChildId.GetKey());
 			if (Entry.IsParentValid() && Entry.IsChildValid())
 			{
 				bool bDidCreateWidget = false;
@@ -346,6 +357,8 @@ void UUIFrameworkPlayerComponent::LocalWidgetWasAddedToTree(const FUIFrameworkWi
 		}
 		else if (!WidgetClass.IsNull() && WidgetClass.IsPending())
 		{
+			UE_LOG(LogUIFramework, VeryVerbose, TEXT("The widget '%" INT64_FMT "' needs to load its WidgetClass soft pointer."), Entry.ChildId.GetKey());
+
 			if (FWidgetClassToLoad* FoundWidgetClassToLoad = ClassesToLoad.Find(WidgetClass))
 			{
 				FoundWidgetClassToLoad->EntryReplicationIds.AddUnique(Entry.ReplicationID);
@@ -378,6 +391,8 @@ void UUIFrameworkPlayerComponent::LocalWidgetWasAddedToTree(const FUIFrameworkWi
 	}
 	else
 	{
+		UE_LOG(LogUIFramework, VeryVerbose, TEXT("Widget '%" INT64_FMT "' is not ready for replication yet."), Entry.ChildId.GetKey());
+
 		NetReplicationPending.Add(Entry.ReplicationID);
 	}
 
@@ -392,6 +407,7 @@ void UUIFrameworkPlayerComponent::LocalWidgetRemovedFromTree(const FUIFrameworkW
 
 	NetReplicationPending.Remove(Entry.ReplicationID);
 	AddPending.Remove(Entry.ReplicationID);
+	UE_LOG(LogUIFramework, VeryVerbose, TEXT("Removed Widget '%" INT64_FMT "' from the widget tree locally."), Entry.ChildId.GetKey());
 
 	PrimaryComponentTick.SetTickFunctionEnable(NetReplicationPending.Num() > 0 || AddPending.Num() > 0 || ClassesToLoad.Num() > 0);
 }
@@ -478,6 +494,8 @@ void UUIFrameworkPlayerComponent::LocalAddChild(FUIFrameworkWidgetId WidgetId)
 				if (UUIFrameworkPresenter* PresenterPtr = GetOrCreatePresenter())
 				{
 					PresenterPtr->AddToViewport(UMGWidget, *LayerSlot);
+					UE_LOG(LogUIFramework, VeryVerbose, TEXT("Widget %s added to the local viewport via the presenter."), *UMGWidget->GetName());
+
 				}
 			}
 		}
diff --git a/Engine/Plugins/Experimental/UIFramework/Source/Private/UIFWidget.cpp b/Engine/Plugins/Experimental/UIFramework/Source/Private/UIFWidget.cpp
--- a/Engine/Plugins/Experimental/UIFramework/Source/Private/UIFWidget.cpp
+++ b/Engine/Plugins/Experimental/UIFramework/Source/Private/UIFWidget.cpp
@@ -13,6 +13,7 @@
 #include "Types/UIFWidgetTree.h"
 #include "Types/UIFWidgetOwner.h"
 #include "Types/UIFWidgetTreeOwner.h"
+#include "UIFLog.h"
 
 #include UE_INLINE_GENERATED_CPP_BY_NAME(UIFWidget)
 
@@ -232,6 +233,17 @@ void UUIFrameworkWidget::LocalCreateUMGWidget(IUIFrameworkWidgetTreeOwner* InOwn
 				LocalOnUMGWidgetCreated();
 			}
 		}
+
+		if (LocalUMGWidget)
+		{
+			UE_LOG(LogUIFramework, VeryVerbose, TEXT("UMG Widget object %s created on the client."), *LocalUMGWidget->GetName());
+
+		}
+		else
+		{ 
+			UE_LOG(LogUIFramework, VeryVerbose, TEXT("No UMG Widget object was created for widget ID  '%" INT64_FMT "'."), Id.GetKey());
+		}
+
 	}
 
 	if (bDidCreateWidgetPtr)
