UIFramework: add logging for UIFPlayerComponent

#Synced-CL 50600560

diff --git a/Engine/Plugins/Experimental/UIFramework/Source/Private/UIFPlayerComponent.cpp b/Engine/Plugins/Experimental/UIFramework/Source/Private/UIFPlayerComponent.cpp
--- a/Engine/Plugins/Experimental/UIFramework/Source/Private/UIFPlayerComponent.cpp
+++ b/Engine/Plugins/Experimental/UIFramework/Source/Private/UIFPlayerComponent.cpp
@@ -85,6 +85,7 @@ UUIFrameworkPlayerComponent::UUIFrameworkPlayerComponent()
 	: RootList(this)
 	, WidgetTree(GetOwner(), this)
 {
+	UE_LOG(LogUIFramework, VeryVerbose, TEXT("Constructing UIFrameworkPlayerComponent"));
 	SetIsReplicatedByDefault(true);
 	bWantsInitializeComponent = true;
 
