Fix crash by direct loading default contant on designer preview, canvas preview creates and discards widgets each mouse event

#Synced-CL 51535387

diff --git a/Engine/Plugins/Experimental/UIFramework/Source/Private/Widgets/UIFrameworkCustomButtonWidget.cpp b/Engine/Plugins/Experimental/UIFramework/Source/Private/Widgets/UIFrameworkCustomButtonWidget.cpp
--- a/Engine/Plugins/Experimental/UIFramework/Source/Private/Widgets/UIFrameworkCustomButtonWidget.cpp
+++ b/Engine/Plugins/Experimental/UIFramework/Source/Private/Widgets/UIFrameworkCustomButtonWidget.cpp
@@ -6,6 +6,7 @@
 #include "Components/ButtonSlot.h"
 #include "Engine/AssetManager.h"
 #include "Engine/StreamableManager.h"
+#include "Widgets/Input/SButton.h"
 
 UUIFrameworkCustomButtonWidget::UUIFrameworkCustomButtonWidget()
 {
@@ -46,30 +47,59 @@ void UUIFrameworkCustomButtonWidget::LocalCreateDefaultUMGWidget()
 		{
 			if (Class->IsChildOf(UUserWidget::StaticClass()))
 			{
-				LocalDefaultUMGWidget = CreateWidget(this, Class);
+				if (UUserWidget* OwningUserWidget = GetTypedOuter<UUserWidget>())
+				{
+					LocalDefaultUMGWidget = CreateWidget(OwningUserWidget, Class);
+				}
 			}
 			else
 			{
 				check(Class->IsChildOf(UWidget::StaticClass()));
 				LocalDefaultUMGWidget = NewObject<UWidget>(this, Class, FName(), RF_Transient);
 			}
-			
-			UPanelSlot* PanelSlot = GetContentSlot();
-			if (!IsValid( PanelSlot ) || PanelSlot->GetContent() == nullptr)
+
+			if (LocalDefaultUMGWidget)
 			{
-				SetContent(LocalDefaultUMGWidget);
-				
-				if (UButtonSlot* ButtonSlot = Cast<UButtonSlot>(LocalDefaultUMGWidget->Slot))
+#if WITH_EDITOR
+				if (IsDesignTime() && MyButton.IsValid())
 				{
-					ButtonSlot->SetPadding(FMargin());
-					ButtonSlot->SetHorizontalAlignment(EHorizontalAlignment::HAlign_Fill);
-					ButtonSlot->SetVerticalAlignment(EVerticalAlignment::VAlign_Fill);
+					const UPanelSlot* PanelSlot = GetContentSlot();
+					if (!IsValid(PanelSlot) || PanelSlot->GetContent() == nullptr)
+					{
+						MyButton->SetContent(LocalDefaultUMGWidget->TakeWidget());
+					}
+				}
+				else
+#endif
+				{
+					const UPanelSlot* PanelSlot = GetContentSlot();
+					if (!IsValid(PanelSlot) || PanelSlot->GetContent() == nullptr)
+					{
+						SetContent(LocalDefaultUMGWidget);
+
+						if (UButtonSlot* ButtonSlot = Cast<UButtonSlot>(LocalDefaultUMGWidget->Slot))
+						{
+							ButtonSlot->SetPadding(FMargin());
+							ButtonSlot->SetHorizontalAlignment(EHorizontalAlignment::HAlign_Fill);
+							ButtonSlot->SetVerticalAlignment(EVerticalAlignment::VAlign_Fill);
+						}
+					}
 				}
 			}
 		}
 	}
 }
 
+void UUIFrameworkCustomButtonWidget::ReleaseSlateResources(bool bReleaseChildren)
+{
+	Super::ReleaseSlateResources(bReleaseChildren);
+	if (WidgetClassStreamableHandle.IsValid())
+	{
+		WidgetClassStreamableHandle->CancelHandle();
+		WidgetClassStreamableHandle.Reset();
+	}	
+}
+
 TSharedPtr<FStreamableHandle> UUIFrameworkCustomButtonWidget::AsyncLoadDefaultWidgetClass()
 {
 	if (WidgetClassStreamableHandle.IsValid())
@@ -87,7 +117,21 @@ TSharedPtr<FStreamableHandle> UUIFrameworkCustomButtonWidget::AsyncLoadDefaultWi
 		WidgetClassStreamableHandle->CancelHandle();
 		WidgetClassStreamableHandle.Reset();
 	}
-	
+
+#if WITH_EDITOR
+	if (IsDesignTime())
+	{
+		if (!DefaultContentWidgetClass.Get() && !DefaultContentWidgetClass.IsNull())
+		{
+			UAssetManager::GetStreamableManager().LoadSynchronous(DefaultContentWidgetClass.ToSoftObjectPath());
+		}
+		
+		LocalCreateDefaultUMGWidget();
+
+		return nullptr;
+	}
+#endif
+
 	TWeakObjectPtr<ThisClass> WeakSelf = this;
 	WidgetClassStreamableHandle = UAssetManager::GetStreamableManager().RequestAsyncLoad(
 		DefaultContentWidgetClass.ToSoftObjectPath()
diff --git a/Engine/Plugins/Experimental/UIFramework/Source/Private/Widgets/UIFrameworkCustomButtonWidget.h b/Engine/Plugins/Experimental/UIFramework/Source/Private/Widgets/UIFrameworkCustomButtonWidget.h
--- a/Engine/Plugins/Experimental/UIFramework/Source/Private/Widgets/UIFrameworkCustomButtonWidget.h
+++ b/Engine/Plugins/Experimental/UIFramework/Source/Private/Widgets/UIFrameworkCustomButtonWidget.h
@@ -26,6 +26,8 @@ public:
 private:
 	void LocalCreateDefaultUMGWidget();
 	
+	virtual void ReleaseSlateResources(bool bReleaseChildren) override;
+	
 	UFUNCTION()
 	void HandleButtonClick();
 	
