Verse UI Input: expose Back input action internally.
Back input action consumes input, and pressing Esc doesn't bring up the social menu while the Back action is bound
#changelist validated
#virtualized
#Synced-CL 51916708

#Synced-CL 51916708

#Synced-CL 51916708

diff --git a/Engine/Plugins/Experimental/UIFramework/Source/Private/Widgets/UIFrameworkButtonWidget.cpp b/Engine/Plugins/Experimental/UIFramework/Source/Private/Widgets/UIFrameworkButtonWidget.cpp
--- a/Engine/Plugins/Experimental/UIFramework/Source/Private/Widgets/UIFrameworkButtonWidget.cpp
+++ b/Engine/Plugins/Experimental/UIFramework/Source/Private/Widgets/UIFrameworkButtonWidget.cpp
@@ -14,6 +14,7 @@
 #include "GameFramework/PlayerController.h"
 #include "Input/CommonUIActionRouterBase.h"
 #include "Input/CommonUIInputTypes.h"
+#include "InputTriggers.h"
 #include "TimerManager.h"
 #include "UIFWidget.h"
 #include "Widgets/DeclarativeSyntaxSupport.h"
@@ -155,8 +156,20 @@ void UUIFrameworkButtonWidget::BindTriggeringInputActionToClick()
 				if (UCommonUIActionRouterBase* ActionRouter = UCommonUIActionRouterBase::Get(*ActivatableWidget))
 				{
 					FBindUIActionArgs BindArgs(TriggeringInputAction.Get(), false, FSimpleDelegate::CreateUObject(this, &UUIFrameworkButtonWidget::HandleButtonClicked));
-					BindArgs.bConsumeInput = false;
-					BindArgs.KeyEvent = GetClickMethod() == EButtonClickMethod::MouseDown ? IE_Pressed : IE_Released;
+					BindArgs.bConsumeInput = TriggeringInputAction->bConsumeInput;
+					BindArgs.PriorityWithinCollection = 1;
+
+					// This is a bit of a hack to enforce Pressed trigger if it is the only trigger in an input action.
+					// We need this to be able to consume and block other actions
+					const bool bHasPressedTrigger = TriggeringInputAction->Triggers.ContainsByPredicate(
+						[](const UInputTrigger* Trigger)
+						{
+							return Cast<UInputTriggerPressed>(Trigger) != nullptr;
+						}
+					);
+					bool ShouldConsumePressed = TriggeringInputAction->bConsumesActionAndAxisMappings && bHasPressedTrigger && TriggeringInputAction->Triggers.Num() == 1;
+
+					BindArgs.KeyEvent = GetClickMethod() == EButtonClickMethod::MouseDown || ShouldConsumePressed ? IE_Pressed : IE_Released;
 					TriggeringBindingHandle = ActionRouter->RegisterUIActionBinding(*this, BindArgs);
 				}
 			}
